MessagingDebugger: Documentation cleanup pass

diff --git a/Engine/Plugins/Messaging/MessagingDebugger/Source/MessagingDebugger/Private/Models/MessagingDebuggerCommands.h b/Engine/Plugins/Messaging/MessagingDebugger/Source/MessagingDebugger/Private/Models/MessagingDebuggerCommands.h
--- a/Engine/Plugins/Messaging/MessagingDebugger/Source/MessagingDebugger/Private/Models/MessagingDebuggerCommands.h
+++ b/Engine/Plugins/Messaging/MessagingDebugger/Source/MessagingDebugger/Private/Models/MessagingDebuggerCommands.h
@@ -40,34 +40,22 @@ public:
 
 public:
 
-	/**
-	 * Holds a command that breaks at the next message.
-	 */
+	/** Holds a command that breaks at the next message. */
 	TSharedPtr<FUICommandInfo> BreakDebugger;
 
-	/**
-	 * Holds a command that clears the message history.
-	 */
+	/** Holds a command that clears the message history. */
 	TSharedPtr<FUICommandInfo> ClearHistory;
 
-	/**
-	 * Holds a command that continues debugging.
-	 */
+	/** Holds a command that continues debugging. */
 	TSharedPtr<FUICommandInfo> ContinueDebugger;
 
-	/**
-	 * Holds a command that starts debugging.
-	 */
+	/** Holds a command that starts debugging. */
 	TSharedPtr<FUICommandInfo> StartDebugger;
 
-	/**
-	 * Holds a command that continues debugging for one step.
-	 */
+	/** Holds a command that continues debugging for one step. */
 	TSharedPtr<FUICommandInfo> StepDebugger;
 
-	/**
-	 * Holds a command that stops debugging.
-	 */
+	/** Holds a command that stops debugging. */
 	TSharedPtr<FUICommandInfo> StopDebugger;
 };
 
diff --git a/Engine/Plugins/Messaging/MessagingDebugger/Source/MessagingDebugger/Private/Models/MessagingDebuggerEndpointFilter.h b/Engine/Plugins/Messaging/MessagingDebugger/Source/MessagingDebugger/Private/Models/MessagingDebuggerEndpointFilter.h
--- a/Engine/Plugins/Messaging/MessagingDebugger/Source/MessagingDebugger/Private/Models/MessagingDebuggerEndpointFilter.h
+++ b/Engine/Plugins/Messaging/MessagingDebugger/Source/MessagingDebugger/Private/Models/MessagingDebuggerEndpointFilter.h
@@ -7,14 +7,10 @@
 #pragma once
 
 
-/**
- * Type definition for shared pointers to instances of FMessagingDebuggerEndpointFilter.
- */
+/** Type definition for shared pointers to instances of FMessagingDebuggerEndpointFilter. */
 typedef TSharedPtr<class FMessagingDebuggerEndpointFilter> FMessagingDebuggerEndpointFilterPtr;
 
-/**
- * Type definition for shared references to instances of FMessagingDebuggerEndpointFilter.
- */
+/** Type definition for shared references to instances of FMessagingDebuggerEndpointFilter. */
 typedef TSharedRef<class FMessagingDebuggerEndpointFilter> FMessagingDebuggerEndpointFilterRef;
 
 
@@ -28,8 +24,7 @@ public:
 	/**
 	 * Filters the specified endpoint based on the current filter settings.
 	 *
-	 * @param EndpointInfo - The endpoint to filter.
-	 *
+	 * @param EndpointInfo The endpoint to filter.
 	 * @return true if the endpoint passed the filter, false otherwise.
 	 */
 	bool FilterEndpoint( const FMessageTracerEndpointInfoPtr& EndpointInfo ) const
@@ -59,7 +54,7 @@ public:
 	/**
 	 * Sets the filter string.
 	 *
-	 * @param InFilterString - The filter string to set.
+	 * @param InFilterString The filter string to set.
 	 */
 	void SetFilterString( const FString& InFilterString )
 	{
diff --git a/Engine/Plugins/Messaging/MessagingDebugger/Source/MessagingDebugger/Private/Models/MessagingDebuggerMessageFilter.h b/Engine/Plugins/Messaging/MessagingDebugger/Source/MessagingDebugger/Private/Models/MessagingDebuggerMessageFilter.h
--- a/Engine/Plugins/Messaging/MessagingDebugger/Source/MessagingDebugger/Private/Models/MessagingDebuggerMessageFilter.h
+++ b/Engine/Plugins/Messaging/MessagingDebugger/Source/MessagingDebugger/Private/Models/MessagingDebuggerMessageFilter.h
@@ -7,14 +7,10 @@
 #pragma once
 
 
-/**
- * Type definition for shared pointers to instances of FMessagingDebuggerMessageFilter.
- */
+/** Type definition for shared pointers to instances of FMessagingDebuggerMessageFilter. */
 typedef TSharedPtr<class FMessagingDebuggerMessageFilter> FMessagingDebuggerMessageFilterPtr;
 
-/**
- * Type definition for shared references to instances of FMessagingDebuggerMessageFilter.
- */
+/** Type definition for shared references to instances of FMessagingDebuggerMessageFilter. */
 typedef TSharedRef<class FMessagingDebuggerMessageFilter> FMessagingDebuggerMessageFilterRef;
 
 
@@ -28,8 +24,7 @@ public:
 	/**
 	 * Filters the specified message based on the current filter settings.
 	 *
-	 * @param MessageInfo - The message to filter.
-	 *
+	 * @param MessageInfo The message to filter.
 	 * @return true if the message passed the filter, false otherwise.
 	 */
 	bool FilterEndpoint( const FMessageTracerMessageInfoPtr& MessageInfo ) const
diff --git a/Engine/Plugins/Messaging/MessagingDebugger/Source/MessagingDebugger/Private/Models/MessagingDebuggerModel.h b/Engine/Plugins/Messaging/MessagingDebugger/Source/MessagingDebugger/Private/Models/MessagingDebuggerModel.h
--- a/Engine/Plugins/Messaging/MessagingDebugger/Source/MessagingDebugger/Private/Models/MessagingDebuggerModel.h
+++ b/Engine/Plugins/Messaging/MessagingDebugger/Source/MessagingDebugger/Private/Models/MessagingDebuggerModel.h
@@ -7,14 +7,10 @@
 #pragma once
 
 
-/**
- * Type definition for shared pointers to instances of FMessagingDebuggerModel.
- */
+/** Type definition for shared pointers to instances of FMessagingDebuggerModel. */
 typedef TSharedPtr<class FMessagingDebuggerModel> FMessagingDebuggerModelPtr;
 
-/**
- * Type definition for shared references to instances of FMessagingDebuggerTypeFilter.
- */
+/** Type definition for shared references to instances of FMessagingDebuggerTypeFilter. */
 typedef TSharedRef<class FMessagingDebuggerModel> FMessagingDebuggerModelRef;
 
 
@@ -59,8 +55,7 @@ public:
 	/** 
 	 * Checks whether messages of the given message endpoint should be visible.
 	 *
-	 * @param Endpoint - The information of the endpoint to check.
-	 *
+	 * @param Endpoint The information of the endpoint to check.
 	 * @return true if the endpoint's messages should be visible, false otherwise.
 	 */
 	bool IsEndpointVisible( const FMessageTracerEndpointInfoRef& EndpointInfo ) const
@@ -71,8 +66,7 @@ public:
 	/**
 	 * Checks whether the given message should be visible.
 	 *
-	 * @param MessageInfo - The information of the message to check.
-	 *
+	 * @param MessageInfo The information of the message to check.
 	 * @return true if the message should be visible, false otherwise.
 	 */
 	bool IsMessageVisible( const FMessageTracerMessageInfoRef& MessageInfo ) const
@@ -102,8 +96,7 @@ public:
 	/**
 	 * Checks whether messages of the given type should be visible.
 	 *
-	 * @param TypeInfo - The information of the message type to check.
-	 *
+	 * @param TypeInfo The information of the message type to check.
 	 * @return true if messages of the given type should be visible, false otherwise.
 	 */
 	bool IsTypeVisible( const FMessageTracerTypeInfoRef& TypeInfo ) const
@@ -114,7 +107,7 @@ public:
 	/**
 	 * Selects the specified endpoint (or none if nullptr).
 	 *
-	 * @param EndpointInfo - The information of the endpoint to select.
+	 * @param EndpointInfo The information of the endpoint to select.
 	 */
 	void SelectEndpoint( const FMessageTracerEndpointInfoPtr& EndpointInfo )
 	{
@@ -128,7 +121,7 @@ public:
 	/**
 	 * Selects the specified message (or none if nullptr).
 	 *
-	 * @param MessageInfo - The information of the message to select.
+	 * @param MessageInfo The information of the message to select.
 	 */
 	void SelectMessage( const FMessageTracerMessageInfoPtr& MessageInfo )
 	{
@@ -142,8 +135,8 @@ public:
 	/**
 	 * Sets whether messages for the given endpoint should be visible.
 	 *
-	 * @param Endpoint - The information for the endpoint.
-	 * @param Visible - Whether the endpoint's messages should be visible or not.
+	 * @param Endpoint The information for the endpoint.
+	 * @param Visible Whether the endpoint's messages should be visible or not.
 	 */
 	void SetEndpointVisibility( const FMessageTracerEndpointInfoRef& EndpointInfo, bool Visible )
 	{
@@ -162,8 +155,8 @@ public:
 	/**
 	 * Sets whether messages for the given message type should be visible.
 	 *
-	 * @param Endpoint - The information for the message type.
-	 * @param Visible - Whether messages of the given type should be visible or not.
+	 * @param Endpoint The information for the message type.
+	 * @param Visible Whether messages of the given type should be visible or not.
 	 */
 	void SetTypeVisibility( const FMessageTracerTypeInfoRef& TypeInfo, bool Visible )
 	{
diff --git a/Engine/Plugins/Messaging/MessagingDebugger/Source/MessagingDebugger/Private/Models/MessagingDebuggerTypeFilter.h b/Engine/Plugins/Messaging/MessagingDebugger/Source/MessagingDebugger/Private/Models/MessagingDebuggerTypeFilter.h
--- a/Engine/Plugins/Messaging/MessagingDebugger/Source/MessagingDebugger/Private/Models/MessagingDebuggerTypeFilter.h
+++ b/Engine/Plugins/Messaging/MessagingDebugger/Source/MessagingDebugger/Private/Models/MessagingDebuggerTypeFilter.h
@@ -28,8 +28,7 @@ public:
 	/**
 	 * Filters the specified message type based on the current filter settings.
 	 *
-	 * @param TypeInfo - The message type to filter.
-	 *
+	 * @param TypeInfo The message type to filter.
 	 * @return true if the endpoint passed the filter, false otherwise.
 	 */
 	bool FilterType( const FMessageTracerTypeInfoPtr& TypeInfo ) const
@@ -61,7 +60,7 @@ public:
 	/**
 	 * Sets the filter string.
 	 *
-	 * @param InFilterString - The filter string to set.
+	 * @param InFilterString The filter string to set.
 	 */
 	void SetFilterString( const FString& InFilterString )
 	{
diff --git a/Engine/Plugins/Messaging/MessagingDebugger/Source/MessagingDebugger/Private/Widgets/Breakpoints/SMessagingBreakpoints.h b/Engine/Plugins/Messaging/MessagingDebugger/Source/MessagingDebugger/Private/Widgets/Breakpoints/SMessagingBreakpoints.h
--- a/Engine/Plugins/Messaging/MessagingDebugger/Source/MessagingDebugger/Private/Widgets/Breakpoints/SMessagingBreakpoints.h
+++ b/Engine/Plugins/Messaging/MessagingDebugger/Source/MessagingDebugger/Private/Widgets/Breakpoints/SMessagingBreakpoints.h
@@ -24,9 +24,9 @@ public:
 	/**
 	 * Construct this widget
 	 *
-	 * @param InArgs - The declaration data for this widget.
-	 * @param InStyle - The visual style to use for this widget.
-	 * @param InTracer - The message tracer.
+	 * @param InArgs The declaration data for this widget.
+	 * @param InStyle The visual style to use for this widget.
+	 * @param InTracer The message tracer.
 	 */
 	void Construct( const FArguments& InArgs, const TSharedRef<ISlateStyle>& InStyle, const IMessageTracerRef& InTracer );
 
@@ -44,7 +44,7 @@ private:
 	TArray<IMessageTracerBreakpointPtr> BreakpointList;
 
 	// Holds the message list view.
-	TSharedPtr<SListView<IMessageTracerBreakpointPtr> > BreakpointListView;
+	TSharedPtr<SListView<IMessageTracerBreakpointPtr>> BreakpointListView;
 
 	// Holds the widget's visual style.
 	TSharedPtr<ISlateStyle> Style;
diff --git a/Engine/Plugins/Messaging/MessagingDebugger/Source/MessagingDebugger/Private/Widgets/Breakpoints/SMessagingBreakpointsTableRow.h b/Engine/Plugins/Messaging/MessagingDebugger/Source/MessagingDebugger/Private/Widgets/Breakpoints/SMessagingBreakpointsTableRow.h
--- a/Engine/Plugins/Messaging/MessagingDebugger/Source/MessagingDebugger/Private/Widgets/Breakpoints/SMessagingBreakpointsTableRow.h
+++ b/Engine/Plugins/Messaging/MessagingDebugger/Source/MessagingDebugger/Private/Widgets/Breakpoints/SMessagingBreakpointsTableRow.h
@@ -28,8 +28,8 @@ public:
 	/**
 	 * Constructs the widget.
 	 *
-	 * @param InArgs - The construction arguments.
-	 * @param InOwnerTableView - The table view that owns this row.
+	 * @param InArgs The construction arguments.
+	 * @param InOwnerTableView The table view that owns this row.
 	 */
 	void Construct( const FArguments& InArgs, const TSharedRef<STableViewBase>& InOwnerTableView )
 	{
diff --git a/Engine/Plugins/Messaging/MessagingDebugger/Source/MessagingDebugger/Private/Widgets/EndpointDetails/SMessagingEndpointDetails.h b/Engine/Plugins/Messaging/MessagingDebugger/Source/MessagingDebugger/Private/Widgets/EndpointDetails/SMessagingEndpointDetails.h
--- a/Engine/Plugins/Messaging/MessagingDebugger/Source/MessagingDebugger/Private/Widgets/EndpointDetails/SMessagingEndpointDetails.h
+++ b/Engine/Plugins/Messaging/MessagingDebugger/Source/MessagingDebugger/Private/Widgets/EndpointDetails/SMessagingEndpointDetails.h
@@ -23,9 +23,9 @@ public:
 	/**
 	 * Construct this widget
 	 *
-	 * @param InArgs - The declaration data for this widget.
-	 * @param InModel - The view model to use.
-	 * @param InStyle - The visual style to use for this widget.
+	 * @param InArgs The declaration data for this widget.
+	 * @param InModel The view model to use.
+	 * @param InStyle The visual style to use for this widget.
 	 */
 	void Construct( const FArguments& InArgs, const FMessagingDebuggerModelRef& InModel, const TSharedRef<ISlateStyle>& InStyle );
 
diff --git a/Engine/Plugins/Messaging/MessagingDebugger/Source/MessagingDebugger/Private/Widgets/Endpoints/SMessagingEndpoints.h b/Engine/Plugins/Messaging/MessagingDebugger/Source/MessagingDebugger/Private/Widgets/Endpoints/SMessagingEndpoints.h
--- a/Engine/Plugins/Messaging/MessagingDebugger/Source/MessagingDebugger/Private/Widgets/Endpoints/SMessagingEndpoints.h
+++ b/Engine/Plugins/Messaging/MessagingDebugger/Source/MessagingDebugger/Private/Widgets/Endpoints/SMessagingEndpoints.h
@@ -30,10 +30,10 @@ public:
 	/**
 	 * Construct this widget
 	 *
-	 * @param InArgs - The declaration data for this widget.
-	 * @param InModel - The view model to use.
-	 * @param InStyle - The visual style to use for this widget.
-	 * @param InTracer - The message tracer.
+	 * @param InArgs The declaration data for this widget.
+	 * @param InModel The view model to use.
+	 * @param InStyle The visual style to use for this widget.
+	 * @param InTracer The message tracer.
 	 */
 	void Construct( const FArguments& InArgs, const FMessagingDebuggerModelRef& InModel, const TSharedRef<ISlateStyle>& InStyle, const IMessageTracerRef& InTracer );
 
@@ -78,7 +78,7 @@ private:
 	TArray<FMessageTracerEndpointInfoPtr> EndpointList;
 
 	// Holds the message list view.
-	TSharedPtr<SListView<FMessageTracerEndpointInfoPtr> > EndpointListView;
+	TSharedPtr<SListView<FMessageTracerEndpointInfoPtr>> EndpointListView;
 
 	// Holds the endpoint filter model.
 	FMessagingDebuggerEndpointFilterPtr Filter;
diff --git a/Engine/Plugins/Messaging/MessagingDebugger/Source/MessagingDebugger/Private/Widgets/Endpoints/SMessagingEndpointsFilterBar.h b/Engine/Plugins/Messaging/MessagingDebugger/Source/MessagingDebugger/Private/Widgets/Endpoints/SMessagingEndpointsFilterBar.h
--- a/Engine/Plugins/Messaging/MessagingDebugger/Source/MessagingDebugger/Private/Widgets/Endpoints/SMessagingEndpointsFilterBar.h
+++ b/Engine/Plugins/Messaging/MessagingDebugger/Source/MessagingDebugger/Private/Widgets/Endpoints/SMessagingEndpointsFilterBar.h
@@ -26,8 +26,8 @@ public:
 	/**
 	 * Construct this widget
 	 *
-	 * @param InArgs - The declaration data for this widget.
-	 * @param InFilter - The filter model.
+	 * @param InArgs The declaration data for this widget.
+	 * @param InFilter The filter model.
 	 */
 	void Construct( const FArguments& InArgs, FMessagingDebuggerEndpointFilterRef InFilter )
 	{
diff --git a/Engine/Plugins/Messaging/MessagingDebugger/Source/MessagingDebugger/Private/Widgets/Endpoints/SMessagingEndpointsTableRow.h b/Engine/Plugins/Messaging/MessagingDebugger/Source/MessagingDebugger/Private/Widgets/Endpoints/SMessagingEndpointsTableRow.h
--- a/Engine/Plugins/Messaging/MessagingDebugger/Source/MessagingDebugger/Private/Widgets/Endpoints/SMessagingEndpointsTableRow.h
+++ b/Engine/Plugins/Messaging/MessagingDebugger/Source/MessagingDebugger/Private/Widgets/Endpoints/SMessagingEndpointsTableRow.h
@@ -29,8 +29,8 @@ public:
 	/**
 	 * Constructs the widget.
 	 *
-	 * @param InArgs - The construction arguments.
-	 * @param InOwnerTableView - The table view that owns this row.
+	 * @param InArgs The construction arguments.
+	 * @param InOwnerTableView The table view that owns this row.
 	 */
 	void Construct( const FArguments& InArgs, const TSharedRef<STableViewBase>& InOwnerTableView, const FMessagingDebuggerModelRef& InModel )
 	{
diff --git a/Engine/Plugins/Messaging/MessagingDebugger/Source/MessagingDebugger/Private/Widgets/Graph/SMessagingGraph.h b/Engine/Plugins/Messaging/MessagingDebugger/Source/MessagingDebugger/Private/Widgets/Graph/SMessagingGraph.h
--- a/Engine/Plugins/Messaging/MessagingDebugger/Source/MessagingDebugger/Private/Widgets/Graph/SMessagingGraph.h
+++ b/Engine/Plugins/Messaging/MessagingDebugger/Source/MessagingDebugger/Private/Widgets/Graph/SMessagingGraph.h
@@ -23,8 +23,8 @@ public:
 	/**
 	 * Construct this widget
 	 *
-	 * @param InArgs - The declaration data for this widget.
-	 * @param InStyle - The visual style to use for this widget.
+	 * @param InArgs The declaration data for this widget.
+	 * @param InStyle The visual style to use for this widget.
 	 */
 	void Construct( const FArguments& InArgs, const TSharedRef<ISlateStyle>& InStyle );
 
diff --git a/Engine/Plugins/Messaging/MessagingDebugger/Source/MessagingDebugger/Private/Widgets/History/SMessagingHistory.h b/Engine/Plugins/Messaging/MessagingDebugger/Source/MessagingDebugger/Private/Widgets/History/SMessagingHistory.h
--- a/Engine/Plugins/Messaging/MessagingDebugger/Source/MessagingDebugger/Private/Widgets/History/SMessagingHistory.h
+++ b/Engine/Plugins/Messaging/MessagingDebugger/Source/MessagingDebugger/Private/Widgets/History/SMessagingHistory.h
@@ -30,10 +30,10 @@ public:
 	/**
 	 * Construct this widget
 	 *
-	 * @param InArgs - The declaration data for this widget.
-	 * @param InModel - The view model to use.
-	 * @param InStyle - The visual style to use for this widget.
-	 * @param InTracer - The message tracer.
+	 * @param InArgs The construction arguments.
+	 * @param InModel The view model to use.
+	 * @param InStyle The visual style to use for this widget.
+	 * @param InTracer The message tracer.
 	 */
 	void Construct( const FArguments& InArgs, const FMessagingDebuggerModelRef& InModel, const TSharedRef<ISlateStyle>& InStyle, const IMessageTracerRef& InTracer );
 
@@ -42,7 +42,7 @@ protected:
 	/**
 	 * Adds the given message to the history.
 	 *
-	 * @param MessageInfo - The information about the message to add.
+	 * @param MessageInfo The information about the message to add.
 	 */
 	void AddMessage( const FMessageTracerMessageInfoRef& MessageInfo );
 
@@ -98,7 +98,7 @@ private:
 	TArray<FMessageTracerMessageInfoPtr> MessageList;
 
 	// Holds the message list view.
-	TSharedPtr<SListView<FMessageTracerMessageInfoPtr> > MessageListView;
+	TSharedPtr<SListView<FMessageTracerMessageInfoPtr>> MessageListView;
 
 	// Holds a pointer to the view model.
 	FMessagingDebuggerModelPtr Model;
diff --git a/Engine/Plugins/Messaging/MessagingDebugger/Source/MessagingDebugger/Private/Widgets/History/SMessagingHistoryFilterBar.h b/Engine/Plugins/Messaging/MessagingDebugger/Source/MessagingDebugger/Private/Widgets/History/SMessagingHistoryFilterBar.h
--- a/Engine/Plugins/Messaging/MessagingDebugger/Source/MessagingDebugger/Private/Widgets/History/SMessagingHistoryFilterBar.h
+++ b/Engine/Plugins/Messaging/MessagingDebugger/Source/MessagingDebugger/Private/Widgets/History/SMessagingHistoryFilterBar.h
@@ -26,8 +26,8 @@ public:
 	/**
 	 * Construct this widget
 	 *
-	 * @param InArgs - The declaration data for this widget.
-	 * @param InFilter - The filter model.
+	 * @param InArgs The declaration data for this widget.
+	 * @param InFilter The filter model.
 	 */
 	void Construct( const FArguments& InArgs, FMessagingDebuggerMessageFilterRef InFilter )
 	{
diff --git a/Engine/Plugins/Messaging/MessagingDebugger/Source/MessagingDebugger/Private/Widgets/History/SMessagingHistoryTableRow.h b/Engine/Plugins/Messaging/MessagingDebugger/Source/MessagingDebugger/Private/Widgets/History/SMessagingHistoryTableRow.h
--- a/Engine/Plugins/Messaging/MessagingDebugger/Source/MessagingDebugger/Private/Widgets/History/SMessagingHistoryTableRow.h
+++ b/Engine/Plugins/Messaging/MessagingDebugger/Source/MessagingDebugger/Private/Widgets/History/SMessagingHistoryTableRow.h
@@ -29,8 +29,8 @@ public:
 	/**
 	 * Constructs the widget.
 	 *
-	 * @param InArgs - The construction arguments.
-	 * @param InOwnerTableView - The table view that owns this row.
+	 * @param InArgs The construction arguments.
+	 * @param InOwnerTableView The table view that owns this row.
 	 */
 	void Construct( const FArguments& InArgs, const TSharedRef<STableViewBase>& InOwnerTableView )
 	{
@@ -259,8 +259,7 @@ protected:
 	/**
 	 * Converts the given time span in seconds to a human readable string.
 	 *
-	 * @param Seconds - The time span to convert.
-	 *
+	 * @param Seconds The time span to convert.
 	 * @return The string representation.
 	 *
 	 * @todo gmp: refactor this into FText::AsTimespan or something like that
diff --git a/Engine/Plugins/Messaging/MessagingDebugger/Source/MessagingDebugger/Private/Widgets/Interceptors/SMessagingInterceptors.h b/Engine/Plugins/Messaging/MessagingDebugger/Source/MessagingDebugger/Private/Widgets/Interceptors/SMessagingInterceptors.h
--- a/Engine/Plugins/Messaging/MessagingDebugger/Source/MessagingDebugger/Private/Widgets/Interceptors/SMessagingInterceptors.h
+++ b/Engine/Plugins/Messaging/MessagingDebugger/Source/MessagingDebugger/Private/Widgets/Interceptors/SMessagingInterceptors.h
@@ -23,9 +23,9 @@ public:
 	/**
 	 * Construct this widget
 	 *
-	 * @param InArgs - The declaration data for this widget.
-	 * @param InStyle - The visual style to use for this widget.
-	 * @param InTracer - The message tracer.
+	 * @param InArgs The declaration data for this widget.
+	 * @param InStyle The visual style to use for this widget.
+	 * @param InTracer The message tracer.
 	 */
 	void Construct( const FArguments& InArgs, const TSharedRef<ISlateStyle>& InStyle, const IMessageTracerRef& InTracer );
 };
diff --git a/Engine/Plugins/Messaging/MessagingDebugger/Source/MessagingDebugger/Private/Widgets/MessageData/SMessagingMessageData.h b/Engine/Plugins/Messaging/MessagingDebugger/Source/MessagingDebugger/Private/Widgets/MessageData/SMessagingMessageData.h
--- a/Engine/Plugins/Messaging/MessagingDebugger/Source/MessagingDebugger/Private/Widgets/MessageData/SMessagingMessageData.h
+++ b/Engine/Plugins/Messaging/MessagingDebugger/Source/MessagingDebugger/Private/Widgets/MessageData/SMessagingMessageData.h
@@ -31,9 +31,9 @@ public:
 	/**
 	 * Construct this widget
 	 *
-	 * @param InArgs - The declaration data for this widget.
-	 * @param InModel - The view model to use.
-	 * @param InStyle - The visual style to use for this widget.
+	 * @param InArgs The declaration data for this widget.
+	 * @param InModel The view model to use.
+	 * @param InStyle The visual style to use for this widget.
 	 */
 	void Construct( const FArguments& InArgs, const FMessagingDebuggerModelRef& InModel, const TSharedRef<ISlateStyle>& InStyle );
 
diff --git a/Engine/Plugins/Messaging/MessagingDebugger/Source/MessagingDebugger/Private/Widgets/MessageDetails/SMessagingDispatchStateTableRow.h b/Engine/Plugins/Messaging/MessagingDebugger/Source/MessagingDebugger/Private/Widgets/MessageDetails/SMessagingDispatchStateTableRow.h
--- a/Engine/Plugins/Messaging/MessagingDebugger/Source/MessagingDebugger/Private/Widgets/MessageDetails/SMessagingDispatchStateTableRow.h
+++ b/Engine/Plugins/Messaging/MessagingDebugger/Source/MessagingDebugger/Private/Widgets/MessageDetails/SMessagingDispatchStateTableRow.h
@@ -28,8 +28,8 @@ public:
 	/**
 	 * Constructs the widget.
 	 *
-	 * @param InArgs - The construction arguments.
-	 * @param InOwnerTableView - The table view that owns this row.
+	 * @param InArgs The construction arguments.
+	 * @param InOwnerTableView The table view that owns this row.
 	 */
 	void Construct( const FArguments& InArgs, const TSharedRef<STableViewBase>& InOwnerTableView, const FMessagingDebuggerModelRef& InModel )
 	{
@@ -104,8 +104,7 @@ protected:
 	/**
 	 * Converts a time span a color value.
 	 *
-	 * @param Latency - The time span to convert.
-	 *
+	 * @param Latency The time span to convert.
 	 * @return The corresponding color value.
 	 */
 	FSlateColor TimespanToColor( double Timespan ) const
@@ -131,8 +130,7 @@ protected:
 	/**
 	 * Converts the given time span in seconds to a human readable string.
 	 *
-	 * @param Seconds - The time span to convert.
-	 *
+	 * @param Seconds The time span to convert.
 	 * @return The text representation.
 	 *
 	 * @todo gmp: refactor this into FText::AsTimespan or something like that
diff --git a/Engine/Plugins/Messaging/MessagingDebugger/Source/MessagingDebugger/Private/Widgets/MessageDetails/SMessagingMessageDetails.h b/Engine/Plugins/Messaging/MessagingDebugger/Source/MessagingDebugger/Private/Widgets/MessageDetails/SMessagingMessageDetails.h
--- a/Engine/Plugins/Messaging/MessagingDebugger/Source/MessagingDebugger/Private/Widgets/MessageDetails/SMessagingMessageDetails.h
+++ b/Engine/Plugins/Messaging/MessagingDebugger/Source/MessagingDebugger/Private/Widgets/MessageDetails/SMessagingMessageDetails.h
@@ -30,9 +30,9 @@ public:
 	/**
 	 * Construct this widget
 	 *
-	 * @param InArgs - The declaration data for this widget.
-	 * @param InModel - The view model to use.
-	 * @param InStyle - The visual style to use for this widget.
+	 * @param InArgs The declaration data for this widget.
+	 * @param InModel The view model to use.
+	 * @param InStyle The visual style to use for this widget.
 	 */
 	void Construct( const FArguments& InArgs, const FMessagingDebuggerModelRef& InModel, const TSharedRef<ISlateStyle>& InStyle );
 
diff --git a/Engine/Plugins/Messaging/MessagingDebugger/Source/MessagingDebugger/Private/Widgets/SMessagingDebugger.h b/Engine/Plugins/Messaging/MessagingDebugger/Source/MessagingDebugger/Private/Widgets/SMessagingDebugger.h
--- a/Engine/Plugins/Messaging/MessagingDebugger/Source/MessagingDebugger/Private/Widgets/SMessagingDebugger.h
+++ b/Engine/Plugins/Messaging/MessagingDebugger/Source/MessagingDebugger/Private/Widgets/SMessagingDebugger.h
@@ -30,11 +30,11 @@ public:
 	/**
 	 * Constructs the application.
 	 *
-	 * @param InArgs - The Slate argument list.
-	 * @param ConstructUnderMajorTab - The major tab which will contain the session front-end.
-	 * @param ConstructUnderWindow - The window in which this widget is being constructed.
-	 * @param InMessageTracer - The message tracer to use for the debugger.
-	 * @param InStyleSet - The style set to use.
+	 * @param InArgs The Slate argument list.
+	 * @param ConstructUnderMajorTab The major tab which will contain the session front-end.
+	 * @param ConstructUnderWindow The window in which this widget is being constructed.
+	 * @param InMessageTracer The message tracer to use for the debugger.
+	 * @param InStyleSet The style set to use.
 	 */
 	void Construct( const FArguments& InArgs, const TSharedRef<SDockTab>& ConstructUnderMajorTab, const TSharedPtr<SWindow>& ConstructUnderWindow, const IMessageTracerRef& InMessageTracer, const TSharedRef<ISlateStyle>& InStyle );
 
@@ -51,10 +51,10 @@ protected:
 	/**
 	 * Fills the Window menu with menu items.
 	 *
-	 * @param MenuBuilder - The multi-box builder that should be filled with content for this pull-down menu.
-	 * @param RootMenuGroup - The root menu group.
-	 * @param AppMenuGroup - The application menu group.
-	 * @param TabManager - A Tab Manager from which to populate tab spawner menu items.
+	 * @param MenuBuilder The multi-box builder that should be filled with content for this pull-down menu.
+	 * @param RootMenuGroup The root menu group.
+	 * @param AppMenuGroup The application menu group.
+	 * @param TabManager A Tab Manager from which to populate tab spawner menu items.
 	 */
 	static void FillWindowMenu( FMenuBuilder& MenuBuilder, TSharedRef<FWorkspaceItem> RootMenuGroup, TSharedRef<FWorkspaceItem> AppMenuGroup, const TSharedPtr<FTabManager> TabManager );
 
diff --git a/Engine/Plugins/Messaging/MessagingDebugger/Source/MessagingDebugger/Private/Widgets/Toolbar/SMessagingDebuggerToolbar.h b/Engine/Plugins/Messaging/MessagingDebugger/Source/MessagingDebugger/Private/Widgets/Toolbar/SMessagingDebuggerToolbar.h
--- a/Engine/Plugins/Messaging/MessagingDebugger/Source/MessagingDebugger/Private/Widgets/Toolbar/SMessagingDebuggerToolbar.h
+++ b/Engine/Plugins/Messaging/MessagingDebugger/Source/MessagingDebugger/Private/Widgets/Toolbar/SMessagingDebuggerToolbar.h
@@ -23,9 +23,9 @@ public:
 	/**
 	 * Construct this widget
 	 *
-	 * @param InArgs - The declaration data for this widget.
-	 * @param InStyle - The visual style to use for this widget.
-	 * @param InCommandList - The command list to bind to.
+	 * @param InArgs The declaration data for this widget.
+	 * @param InStyle The visual style to use for this widget.
+	 * @param InCommandList The command list to bind to.
 	 */
 	void Construct( const FArguments& InArgs, const TSharedRef<ISlateStyle>& InStyle, const TSharedRef<FUICommandList>& InCommandList );
 
@@ -34,8 +34,7 @@ protected:
 	/**
 	 * Creates the toolbar widget.
 	 *
-	 * @param CommandList - The command list to use for the toolbar buttons.
-	 *
+	 * @param CommandList The command list to use for the toolbar buttons.
 	 * @return The toolbar widget.
 	 */
 	TSharedRef<SWidget> MakeToolbar( const TSharedRef<FUICommandList>& CommandList );
diff --git a/Engine/Plugins/Messaging/MessagingDebugger/Source/MessagingDebugger/Private/Widgets/Types/SMessagingTypes.h b/Engine/Plugins/Messaging/MessagingDebugger/Source/MessagingDebugger/Private/Widgets/Types/SMessagingTypes.h
--- a/Engine/Plugins/Messaging/MessagingDebugger/Source/MessagingDebugger/Private/Widgets/Types/SMessagingTypes.h
+++ b/Engine/Plugins/Messaging/MessagingDebugger/Source/MessagingDebugger/Private/Widgets/Types/SMessagingTypes.h
@@ -30,10 +30,10 @@ public:
 	/**
 	 * Construct this widget
 	 *
-	 * @param InArgs - The declaration data for this widget.
-	 * @param InModel - The view model to use.
-	 * @param InStyle - The visual style to use for this widget.
-	 * @param InTracer - The message tracer.
+	 * @param InArgs The declaration data for this widget.
+	 * @param InModel The view model to use.
+	 * @param InStyle The visual style to use for this widget.
+	 * @param InTracer The message tracer.
 	 */
 	void Construct( const FArguments& InArgs, const FMessagingDebuggerModelRef& InModel, const TSharedRef<ISlateStyle>& InStyle, const IMessageTracerRef& InTracer );
 
@@ -42,7 +42,7 @@ protected:
 	/**
 	 * Adds the given message to the history.
 	 *
-	 * @param TypeInfo - The information about the message type to add.
+	 * @param TypeInfo The information about the message type to add.
 	 */
 	void AddType( const FMessageTracerTypeInfoRef& TypeInfo );
 
@@ -92,5 +92,5 @@ private:
 	TArray<FMessageTracerTypeInfoPtr> TypeList;
 
 	// Holds the message type list view.
-	TSharedPtr<SListView<FMessageTracerTypeInfoPtr> > TypeListView;
+	TSharedPtr<SListView<FMessageTracerTypeInfoPtr>> TypeListView;
 };
diff --git a/Engine/Plugins/Messaging/MessagingDebugger/Source/MessagingDebugger/Private/Widgets/Types/SMessagingTypesFilterBar.h b/Engine/Plugins/Messaging/MessagingDebugger/Source/MessagingDebugger/Private/Widgets/Types/SMessagingTypesFilterBar.h
--- a/Engine/Plugins/Messaging/MessagingDebugger/Source/MessagingDebugger/Private/Widgets/Types/SMessagingTypesFilterBar.h
+++ b/Engine/Plugins/Messaging/MessagingDebugger/Source/MessagingDebugger/Private/Widgets/Types/SMessagingTypesFilterBar.h
@@ -20,9 +20,7 @@ public:
 
 	SLATE_BEGIN_ARGS(SMessagingTypesFilterBar) { }
 
-		/**
-		 * Called when the filter settings have changed.
-		 */
+		/** Called when the filter settings have changed. */
 		SLATE_EVENT(FSimpleDelegate, OnFilterChanged)
 
 	SLATE_END_ARGS()
@@ -41,8 +39,8 @@ public:
 	/**
 	 * Construct this widget
 	 *
-	 * @param InArgs - The declaration data for this widget.
-	 * @param InFilter - The filter model.
+	 * @param InArgs The declaration data for this widget.
+	 * @param InFilter The filter model.
 	 */
 	void Construct( const FArguments& InArgs, FMessagingDebuggerTypeFilterRef InFilter )
 	{
diff --git a/Engine/Plugins/Messaging/MessagingDebugger/Source/MessagingDebugger/Private/Widgets/Types/SMessagingTypesTableRow.h b/Engine/Plugins/Messaging/MessagingDebugger/Source/MessagingDebugger/Private/Widgets/Types/SMessagingTypesTableRow.h
--- a/Engine/Plugins/Messaging/MessagingDebugger/Source/MessagingDebugger/Private/Widgets/Types/SMessagingTypesTableRow.h
+++ b/Engine/Plugins/Messaging/MessagingDebugger/Source/MessagingDebugger/Private/Widgets/Types/SMessagingTypesTableRow.h
@@ -29,8 +29,8 @@ public:
 	/**
 	 * Constructs the widget.
 	 *
-	 * @param InArgs - The construction arguments.
-	 * @param InOwnerTableView - The table view that owns this row.
+	 * @param InArgs The construction arguments.
+	 * @param InOwnerTableView The table view that owns this row.
 	 */
 	void Construct( const FArguments& InArgs, const TSharedRef<STableViewBase>& InOwnerTableView, const FMessagingDebuggerModelRef& InModel )
 	{
diff --git a/Engine/Programs/UnrealBuildTool/BuildConfiguration.xml b/Engine/Programs/UnrealBuildTool/BuildConfiguration.xml
--- a/Engine/Programs/UnrealBuildTool/BuildConfiguration.xml
+++ b/Engine/Programs/UnrealBuildTool/BuildConfiguration.xml
@@ -75,7 +75,7 @@
 		<bUsePDBFiles>False</bUsePDBFiles>
 		<bUsePCHFiles>True</bUsePCHFiles>
 		<bUseSharedPCHs>True</bUseSharedPCHs>
-		<bUseUnityBuild>True</bUseUnityBuild>
+		<bUseUnityBuild>False</bUseUnityBuild>
 		<bForceUnityBuild>False</bForceUnityBuild>
 		<MinGameModuleSourceFilesForUnityBuild>32</MinGameModuleSourceFilesForUnityBuild>
 		<NumIncludedBytesPerUnityCPP>262144</NumIncludedBytesPerUnityCPP>
