Messaging: Commented out PropertyEditor dependencies in Messaging Debugger

We cannot use the property editor until is refactored and decoupled from the Editor.

diff --git a/Engine/Plugins/Messaging/MessagingDebugger/Source/MessagingDebugger/MessagingDebugger.Build.cs b/Engine/Plugins/Messaging/MessagingDebugger/Source/MessagingDebugger/MessagingDebugger.Build.cs
--- a/Engine/Plugins/Messaging/MessagingDebugger/Source/MessagingDebugger/MessagingDebugger.Build.cs
+++ b/Engine/Plugins/Messaging/MessagingDebugger/Source/MessagingDebugger/MessagingDebugger.Build.cs
@@ -11,11 +11,11 @@ public class MessagingDebugger : ModuleRules
 			{
 				"Core",
 				"CoreUObject",
+				"EditorStyle",
                 "InputCore",
-                "PropertyEditor",
+//                "PropertyEditor",
 				"Slate",
 				"SlateCore",
-                "EditorStyle"
 			}
 		);
 
diff --git a/Engine/Plugins/Messaging/MessagingDebugger/Source/MessagingDebugger/Private/MessagingDebuggerPrivatePCH.h b/Engine/Plugins/Messaging/MessagingDebugger/Source/MessagingDebugger/Private/MessagingDebuggerPrivatePCH.h
--- a/Engine/Plugins/Messaging/MessagingDebugger/Source/MessagingDebugger/Private/MessagingDebuggerPrivatePCH.h
+++ b/Engine/Plugins/Messaging/MessagingDebugger/Source/MessagingDebugger/Private/MessagingDebuggerPrivatePCH.h
@@ -11,7 +11,6 @@
 #include "Slate.h"
 #include "SlateStyle.h"
 #include "EditorStyle.h"
-#include "PropertyEditing.h"
 
 
 /* Private includes
diff --git a/Engine/Plugins/Messaging/MessagingDebugger/Source/MessagingDebugger/Private/Widgets/MessageData/SMessagingMessageData.cpp b/Engine/Plugins/Messaging/MessagingDebugger/Source/MessagingDebugger/Private/Widgets/MessageData/SMessagingMessageData.cpp
--- a/Engine/Plugins/Messaging/MessagingDebugger/Source/MessagingDebugger/Private/Widgets/MessageData/SMessagingMessageData.cpp
+++ b/Engine/Plugins/Messaging/MessagingDebugger/Source/MessagingDebugger/Private/Widgets/MessageData/SMessagingMessageData.cpp
@@ -27,7 +27,7 @@ void SMessagingMessageData::Construct( const FArguments& InArgs, const FMessagin
 	Style = InStyle;
 
 	// initialize details view
-	FDetailsViewArgs DetailsViewArgs;
+/*	FDetailsViewArgs DetailsViewArgs;
 	{
 		DetailsViewArgs.bAllowSearch = false;
 		DetailsViewArgs.bHideSelectionTip = true;
@@ -43,7 +43,7 @@ void SMessagingMessageData::Construct( const FArguments& InArgs, const FMessagin
 	DetailsView = FModuleManager::GetModuleChecked<FPropertyEditorModule>("PropertyEditor").CreateDetailView(DetailsViewArgs);
 	DetailsView->SetEnabled(TAttribute<bool>::Create(TAttribute<bool>::FGetter::CreateSP(this, &SMessagingMessageData::HandleDetailsViewEnabled)));
 	DetailsView->SetVisibility(TAttribute<EVisibility>::Create(TAttribute<EVisibility>::FGetter::CreateSP(this, &SMessagingMessageData::HandleDetailsViewVisibility)));
-
+*/
 	ChildSlot
 		.HAlign(HAlign_Center)
 		.VAlign(VAlign_Center)
@@ -95,7 +95,7 @@ void SMessagingMessageData::HandleModelSelectedMessageChanged( )
 	}
 	else
 	{
-		DetailsView->SetObject(nullptr);
+//		DetailsView->SetObject(nullptr);
 	}
 }
 
diff --git a/Engine/Plugins/Messaging/MessagingDebugger/Source/MessagingDebugger/Private/Widgets/MessageData/SMessagingMessageData.h b/Engine/Plugins/Messaging/MessagingDebugger/Source/MessagingDebugger/Private/Widgets/MessageData/SMessagingMessageData.h
--- a/Engine/Plugins/Messaging/MessagingDebugger/Source/MessagingDebugger/Private/Widgets/MessageData/SMessagingMessageData.h
+++ b/Engine/Plugins/Messaging/MessagingDebugger/Source/MessagingDebugger/Private/Widgets/MessageData/SMessagingMessageData.h
@@ -53,7 +53,7 @@ private:
 private:
 
 	// Holds the details view.
-	TSharedPtr<IDetailsView> DetailsView;
+//	TSharedPtr<IDetailsView> DetailsView;
 
 	// Holds a pointer to the view model.
 	FMessagingDebuggerModelPtr Model;
