Messaging: Implemented message data panel

diff --git a/Engine/Plugins/Messaging/MessagingDebugger/Source/MessagingDebugger/MessagingDebugger.Build.cs b/Engine/Plugins/Messaging/MessagingDebugger/Source/MessagingDebugger/MessagingDebugger.Build.cs
--- a/Engine/Plugins/Messaging/MessagingDebugger/Source/MessagingDebugger/MessagingDebugger.Build.cs
+++ b/Engine/Plugins/Messaging/MessagingDebugger/Source/MessagingDebugger/MessagingDebugger.Build.cs
@@ -12,6 +12,7 @@ public class MessagingDebugger : ModuleRules
 				"Core",
 				"CoreUObject",
                 "InputCore",
+				"Serialization",
 				"Slate",
 				"SlateCore",
 			}
diff --git a/Engine/Plugins/Messaging/MessagingDebugger/Source/MessagingDebugger/Private/Widgets/MessageData/SMessagingMessageData.cpp b/Engine/Plugins/Messaging/MessagingDebugger/Source/MessagingDebugger/Private/Widgets/MessageData/SMessagingMessageData.cpp
--- a/Engine/Plugins/Messaging/MessagingDebugger/Source/MessagingDebugger/Private/Widgets/MessageData/SMessagingMessageData.cpp
+++ b/Engine/Plugins/Messaging/MessagingDebugger/Source/MessagingDebugger/Private/Widgets/MessageData/SMessagingMessageData.cpp
@@ -1,6 +1,9 @@
 // Copyright 1998-2014 Epic Games, Inc. All Rights Reserved.
 
 #include "MessagingDebuggerPrivatePCH.h"
+#include "Json.h"
+#include "JsonStructSerializerBackend.h"
+#include "StructSerializer.h"
 
 
 #define LOCTEXT_NAMESPACE "SMessagingMessageData"
@@ -45,12 +48,10 @@ void SMessagingMessageData::Construct( const FArguments& InArgs, const FMessagin
 	DetailsView->SetVisibility(TAttribute<EVisibility>::Create(TAttribute<EVisibility>::FGetter::CreateSP(this, &SMessagingMessageData::HandleDetailsViewVisibility)));
 */
 	ChildSlot
-		.HAlign(HAlign_Center)
-		.VAlign(VAlign_Center)
 		[
 			//DetailsView.ToSharedRef()
-			SNew(STextBlock)
-				.Text(LOCTEXT("NotImplementedYet", "Not implemented yet"))
+			SAssignNew(TextBox, SMultiLineEditableTextBox)
+				.IsReadOnly(true)
 		];
 
 	Model->OnSelectedMessageChanged().AddRaw(this, &SMessagingMessageData::HandleModelSelectedMessageChanged);
@@ -91,11 +92,29 @@ void SMessagingMessageData::HandleModelSelectedMessageChanged( )
 
 	if (SelectedMessage.IsValid() && SelectedMessage->Context.IsValid())
 	{
-		// @todo gmp: add support for displaying UScriptStructs in details view
+		UScriptStruct* MessageTypeInfo = SelectedMessage->Context->GetMessageTypeInfo().Get();
+
+		if (MessageTypeInfo != nullptr)
+		{
+			FBufferArchive BufferArchive;
+			FJsonStructSerializerBackend Backend(BufferArchive);
+
+			FStructSerializer::Serialize(SelectedMessage->Context->GetMessage(), *MessageTypeInfo, Backend);
+
+			// add string terminator
+			BufferArchive.Add(0);
+			BufferArchive.Add(0);
+
+			TextBox->SetText(FText::FromString(FString((TCHAR*)BufferArchive.GetData()).Replace(TEXT("\t"), TEXT("    "))));
+		}
+		else
+		{
+			TextBox->SetText(FText::Format(LOCTEXT("UnknownMessageTypeFormat", "Unknown message type '{0}'"), FText::FromString(SelectedMessage->Context->GetMessageType().ToString())));
+		}
 	}
 	else
 	{
-//		DetailsView->SetObject(nullptr);
+		TextBox->SetText(FText::GetEmpty());
 	}
 }
 
diff --git a/Engine/Plugins/Messaging/MessagingDebugger/Source/MessagingDebugger/Private/Widgets/MessageData/SMessagingMessageData.h b/Engine/Plugins/Messaging/MessagingDebugger/Source/MessagingDebugger/Private/Widgets/MessageData/SMessagingMessageData.h
--- a/Engine/Plugins/Messaging/MessagingDebugger/Source/MessagingDebugger/Private/Widgets/MessageData/SMessagingMessageData.h
+++ b/Engine/Plugins/Messaging/MessagingDebugger/Source/MessagingDebugger/Private/Widgets/MessageData/SMessagingMessageData.h
@@ -17,9 +17,7 @@ public:
 
 public:
 
-	/**
-	 * Destructor.
-	 */
+	/** Destructor. */
 	~SMessagingMessageData( );
 
 public:
@@ -41,25 +39,28 @@ public:
 
 private:
 
-	// Handles checking whether the details view is enabled.
+	/** Handles checking whether the details view is enabled. */
 	bool HandleDetailsViewEnabled( ) const;
 
-	// Handles determining the visibility of the details view.
+	/** Handles determining the visibility of the details view. */
 	EVisibility HandleDetailsViewVisibility( ) const;
 
-	// Callback for handling message selection changes.
+	/** Callback for handling message selection changes. */
 	void HandleModelSelectedMessageChanged( );
 
 private:
 
-	// Holds the details view.
+	/** Holds the details view. */
 //	TSharedPtr<IDetailsView> DetailsView;
 
-	// Holds a pointer to the view model.
+	/** Holds a pointer to the view model. */
 	FMessagingDebuggerModelPtr Model;
 
-	// Holds the widget's visual style.
+	/** Holds the widget's visual style. */
 	TSharedPtr<ISlateStyle> Style;
+
+	/** Holds the details text box. */
+	TSharedPtr<SMultiLineEditableTextBox> TextBox;
 };
 
 
diff --git a/Engine/Plugins/Messaging/MessagingDebugger/Source/MessagingDebugger/Private/Widgets/SMessagingDebugger.cpp b/Engine/Plugins/Messaging/MessagingDebugger/Source/MessagingDebugger/Private/Widgets/SMessagingDebugger.cpp
--- a/Engine/Plugins/Messaging/MessagingDebugger/Source/MessagingDebugger/Private/Widgets/SMessagingDebugger.cpp
+++ b/Engine/Plugins/Messaging/MessagingDebugger/Source/MessagingDebugger/Private/Widgets/SMessagingDebugger.cpp
@@ -123,7 +123,7 @@ void SMessagingDebugger::Construct( const FArguments& InArgs, const TSharedRef<S
 						(
 							FTabManager::NewStack()
 								->AddTab(EndpointsTabId, ETabState::OpenedTab)
-								->SetSizeCoefficient(0.65f)
+								->SetSizeCoefficient(0.6f)
 						)
 						->Split
 						(
@@ -131,7 +131,7 @@ void SMessagingDebugger::Construct( const FArguments& InArgs, const TSharedRef<S
 								->AddTab(EndpointDetailsTabId, ETabState::OpenedTab)
 								->AddTab(InterceptorsTabId, ETabState::OpenedTab)
 								->SetForegroundTab(EndpointDetailsTabId)
-								->SetSizeCoefficient(0.35f)
+								->SetSizeCoefficient(0.4f)
 						)
 				)
 				->Split
@@ -176,7 +176,7 @@ void SMessagingDebugger::Construct( const FArguments& InArgs, const TSharedRef<S
 						(
 							FTabManager::NewStack()
 								->AddTab(MessageTypesTabId, ETabState::OpenedTab)
-								->SetSizeCoefficient(0.65f)
+								->SetSizeCoefficient(0.6f)
 						)
 						->Split
 						(
@@ -184,7 +184,7 @@ void SMessagingDebugger::Construct( const FArguments& InArgs, const TSharedRef<S
 								->AddTab(MessageDetailsTabId, ETabState::OpenedTab)
 								->AddTab(MessageDataTabId, ETabState::OpenedTab)
 								->SetForegroundTab(MessageDetailsTabId)
-								->SetSizeCoefficient(0.35f)
+								->SetSizeCoefficient(0.4f)
 						)
 				)
 		);
