Removed WorkspaceMenu dependency from MessagingDebugger

diff --git a/Engine/Plugins/Messaging/MessagingDebugger/Source/MessagingDebugger/MessagingDebugger.Build.cs b/Engine/Plugins/Messaging/MessagingDebugger/Source/MessagingDebugger/MessagingDebugger.Build.cs
--- a/Engine/Plugins/Messaging/MessagingDebugger/Source/MessagingDebugger/MessagingDebugger.Build.cs
+++ b/Engine/Plugins/Messaging/MessagingDebugger/Source/MessagingDebugger/MessagingDebugger.Build.cs
@@ -15,7 +15,6 @@ public class MessagingDebugger : ModuleRules
 				"Serialization",
 				"Slate",
 				"SlateCore",
-                "WorkspaceMenuStructure",
 			}
 		);
 
diff --git a/Engine/Plugins/Messaging/MessagingDebugger/Source/MessagingDebugger/Private/Widgets/SMessagingDebugger.cpp b/Engine/Plugins/Messaging/MessagingDebugger/Source/MessagingDebugger/Private/Widgets/SMessagingDebugger.cpp
--- a/Engine/Plugins/Messaging/MessagingDebugger/Source/MessagingDebugger/Private/Widgets/SMessagingDebugger.cpp
+++ b/Engine/Plugins/Messaging/MessagingDebugger/Source/MessagingDebugger/Private/Widgets/SMessagingDebugger.cpp
@@ -238,10 +238,6 @@ void SMessagingDebugger::FillWindowMenu(FMenuBuilder& MenuBuilder, const TShared
 		return;
 	}
 
-#if !WITH_EDITOR
-	FGlobalTabmanager::Get()->PopulateTabSpawnerMenu(MenuBuilder, WorkspaceMenu::GetMenuStructure().GetStructureRoot());
-#endif //!WITH_EDITOR
-
 	TabManager->PopulateLocalTabSpawnerMenu(MenuBuilder);
 }
 
