Removing rogue WorkspaceMenuStructure header

diff --git a/Engine/Plugins/Messaging/MessagingDebugger/Source/MessagingDebugger/Private/MessagingDebuggerModule.cpp b/Engine/Plugins/Messaging/MessagingDebugger/Source/MessagingDebugger/Private/MessagingDebuggerModule.cpp
--- a/Engine/Plugins/Messaging/MessagingDebugger/Source/MessagingDebugger/Private/MessagingDebuggerModule.cpp
+++ b/Engine/Plugins/Messaging/MessagingDebugger/Source/MessagingDebugger/Private/MessagingDebuggerModule.cpp
@@ -2,7 +2,6 @@
 
 #include "MessagingDebuggerPrivatePCH.h"
 #include "Runtime/Core/Public/Features/IModularFeatures.h"
-#include "WorkspaceMenuStructureModule.h"
 #include "SDockTab.h"
 
 
