MessagingDebugger: Fixed padding of 'show all' hyperlink

diff --git a/Engine/Plugins/Messaging/MessagingDebugger/Source/MessagingDebugger/Private/Widgets/History/SMessagingHistory.cpp b/Engine/Plugins/Messaging/MessagingDebugger/Source/MessagingDebugger/Private/Widgets/History/SMessagingHistory.cpp
--- a/Engine/Plugins/Messaging/MessagingDebugger/Source/MessagingDebugger/Private/Widgets/History/SMessagingHistory.cpp
+++ b/Engine/Plugins/Messaging/MessagingDebugger/Source/MessagingDebugger/Private/Widgets/History/SMessagingHistory.cpp
@@ -155,6 +155,7 @@ void SMessagingHistory::Construct( const FArguments& InArgs, const FMessagingDeb
 						+ SHorizontalBox::Slot()
 							.FillWidth(1.0f)
 							.HAlign(HAlign_Left)
+							.Padding(8.0f, 0.0f, 0.0f, 0.0f)
 							[
 								SNew(SHyperlink)
 									.OnNavigate(this, &SMessagingHistory::HandleShowHiddenHyperlinkNavigate)
