MessagingDebugger: Documentation cleanup pass

diff --git a/Engine/Plugins/Messaging/MessagingDebugger/Source/MessagingDebugger/Private/MessagingDebuggerModule.cpp b/Engine/Plugins/Messaging/MessagingDebugger/Source/MessagingDebugger/Private/MessagingDebuggerModule.cpp
--- a/Engine/Plugins/Messaging/MessagingDebugger/Source/MessagingDebugger/Private/MessagingDebuggerModule.cpp
+++ b/Engine/Plugins/Messaging/MessagingDebugger/Source/MessagingDebugger/Private/MessagingDebuggerModule.cpp
@@ -22,7 +22,7 @@ public:
 
 	// IModuleInterface interface
 	
-	virtual void StartupModule( ) override
+	virtual void StartupModule() override
 	{
 		Style = MakeShareable(new FMessagingDebuggerStyle());
 
@@ -35,7 +35,7 @@ public:
 		FGlobalTabmanager::Get()->RegisterTabSpawner(MessagingDebuggerTabName, FOnSpawnTab::CreateRaw(this, &FMessagingDebuggerModule::SpawnMessagingDebuggerTab));
 	}
 
-	virtual void ShutdownModule( ) override
+	virtual void ShutdownModule() override
 	{
 		FGlobalTabmanager::Get()->UnregisterTabSpawner(MessagingDebuggerTabName);
 
@@ -48,8 +48,7 @@ private:
 	/**
 	 * Creates a new messaging debugger tab.
 	 *
-	 * @param SpawnTabArgs - The arguments for the tab to spawn.
-	 *
+	 * @param SpawnTabArgs The arguments for the tab to spawn.
 	 * @return The spawned tab.
 	 */
 	TSharedRef<SDockTab> SpawnMessagingDebuggerTab( const FSpawnTabArgs& SpawnTabArgs )
@@ -77,7 +76,7 @@ private:
 
 private:
 
-	// Holds the plug-ins style set.
+	/** Holds the plug-ins style set. */
 	TSharedPtr<ISlateStyle> Style;
 };
 
diff --git a/Engine/Plugins/Messaging/MessagingDebugger/Source/MessagingDebugger/Private/Models/MessagingDebuggerCommands.h b/Engine/Plugins/Messaging/MessagingDebugger/Source/MessagingDebugger/Private/Models/MessagingDebuggerCommands.h
--- a/Engine/Plugins/Messaging/MessagingDebugger/Source/MessagingDebugger/Private/Models/MessagingDebuggerCommands.h
+++ b/Engine/Plugins/Messaging/MessagingDebugger/Source/MessagingDebugger/Private/Models/MessagingDebuggerCommands.h
@@ -11,9 +11,7 @@ class FMessagingDebuggerCommands
 {
 public:
 
-	/**
-	 * Default constructor.
-	 */
+	/** Default constructor. */
 	FMessagingDebuggerCommands()
 		: TCommands<FMessagingDebuggerCommands>("MessagingDebugger", NSLOCTEXT("Contexts", "MessagingDebugger", "Messaging Debugger"), NAME_None, "MessagingDebuggerStyle")
 	{ }
@@ -22,7 +20,7 @@ public:
 
 	// TCommands interface
 
-	virtual void RegisterCommands( ) override
+	virtual void RegisterCommands() override
 	{
 		UI_COMMAND(BreakDebugger, "Break", "Break the debugger at the next message", EUserInterfaceActionType::Button, FInputGesture());
 		UI_COMMAND(ClearHistory, "Clear History", "Clears the message history list", EUserInterfaceActionType::Button, FInputGesture());
diff --git a/Engine/Plugins/Messaging/MessagingDebugger/Source/MessagingDebugger/Private/Models/MessagingDebuggerEndpointFilter.h b/Engine/Plugins/Messaging/MessagingDebugger/Source/MessagingDebugger/Private/Models/MessagingDebuggerEndpointFilter.h
--- a/Engine/Plugins/Messaging/MessagingDebugger/Source/MessagingDebugger/Private/Models/MessagingDebuggerEndpointFilter.h
+++ b/Engine/Plugins/Messaging/MessagingDebugger/Source/MessagingDebugger/Private/Models/MessagingDebuggerEndpointFilter.h
@@ -3,13 +3,6 @@
 #pragma once
 
 
-/** Type definition for shared pointers to instances of FMessagingDebuggerEndpointFilter. */
-typedef TSharedPtr<class FMessagingDebuggerEndpointFilter> FMessagingDebuggerEndpointFilterPtr;
-
-/** Type definition for shared references to instances of FMessagingDebuggerEndpointFilter. */
-typedef TSharedRef<class FMessagingDebuggerEndpointFilter> FMessagingDebuggerEndpointFilterRef;
-
-
 /**
  * Implements a filter for the message endpoints list.
  */
@@ -66,18 +59,25 @@ public:
 	 * @return The event delegate.
 	 */
 	DECLARE_EVENT(FMessagingDebuggerEndpointFilter, FOnMessagingEndpointFilterChanged);
-	FOnMessagingEndpointFilterChanged& OnChanged( )
+	FOnMessagingEndpointFilterChanged& OnChanged()
 	{
 		return ChangedEvent;
 	}
 
 private:
 
-	// Holds the filter string used to filter endpoints by their names.
+	/** Holds the filter string used to filter endpoints by their names. */
 	FString FilterString;
 
 private:
 
-	// Holds an event delegate that is invoked when the filter settings changed.
+	/** Holds an event delegate that is invoked when the filter settings changed. */
 	FOnMessagingEndpointFilterChanged ChangedEvent;
 };
+
+
+/** Type definition for shared pointers to instances of FMessagingDebuggerEndpointFilter. */
+typedef TSharedPtr<FMessagingDebuggerEndpointFilter> FMessagingDebuggerEndpointFilterPtr;
+
+/** Type definition for shared references to instances of FMessagingDebuggerEndpointFilter. */
+typedef TSharedRef<FMessagingDebuggerEndpointFilter> FMessagingDebuggerEndpointFilterRef;
diff --git a/Engine/Plugins/Messaging/MessagingDebugger/Source/MessagingDebugger/Private/Models/MessagingDebuggerMessageFilter.h b/Engine/Plugins/Messaging/MessagingDebugger/Source/MessagingDebugger/Private/Models/MessagingDebuggerMessageFilter.h
--- a/Engine/Plugins/Messaging/MessagingDebugger/Source/MessagingDebugger/Private/Models/MessagingDebuggerMessageFilter.h
+++ b/Engine/Plugins/Messaging/MessagingDebugger/Source/MessagingDebugger/Private/Models/MessagingDebuggerMessageFilter.h
@@ -3,13 +3,6 @@
 #pragma once
 
 
-/** Type definition for shared pointers to instances of FMessagingDebuggerMessageFilter. */
-typedef TSharedPtr<class FMessagingDebuggerMessageFilter> FMessagingDebuggerMessageFilterPtr;
-
-/** Type definition for shared references to instances of FMessagingDebuggerMessageFilter. */
-typedef TSharedRef<class FMessagingDebuggerMessageFilter> FMessagingDebuggerMessageFilterRef;
-
-
 /**
  * Implements a filter for the message history list.
  */
@@ -43,13 +36,20 @@ public:
 	 * @return The event delegate.
 	 */
 	DECLARE_EVENT(FMessagingDebuggerMessageFilter, FOnMessagingMessageFilterChanged);
-	FOnMessagingMessageFilterChanged& OnChanged( )
+	FOnMessagingMessageFilterChanged& OnChanged()
 	{
 		return ChangedEvent;
 	}
 
 private:
 
-	// Holds an event delegate that is invoked when the filter settings changed.
+	/** Holds an event delegate that is invoked when the filter settings changed. */
 	FOnMessagingMessageFilterChanged ChangedEvent;
 };
+
+
+/** Type definition for shared pointers to instances of FMessagingDebuggerMessageFilter. */
+typedef TSharedPtr<FMessagingDebuggerMessageFilter> FMessagingDebuggerMessageFilterPtr;
+
+/** Type definition for shared references to instances of FMessagingDebuggerMessageFilter. */
+typedef TSharedRef<FMessagingDebuggerMessageFilter> FMessagingDebuggerMessageFilterRef;
diff --git a/Engine/Plugins/Messaging/MessagingDebugger/Source/MessagingDebugger/Private/Models/MessagingDebuggerModel.h b/Engine/Plugins/Messaging/MessagingDebugger/Source/MessagingDebugger/Private/Models/MessagingDebuggerModel.h
--- a/Engine/Plugins/Messaging/MessagingDebugger/Source/MessagingDebugger/Private/Models/MessagingDebuggerModel.h
+++ b/Engine/Plugins/Messaging/MessagingDebugger/Source/MessagingDebugger/Private/Models/MessagingDebuggerModel.h
@@ -3,13 +3,6 @@
 #pragma once
 
 
-/** Type definition for shared pointers to instances of FMessagingDebuggerModel. */
-typedef TSharedPtr<class FMessagingDebuggerModel> FMessagingDebuggerModelPtr;
-
-/** Type definition for shared references to instances of FMessagingDebuggerTypeFilter. */
-typedef TSharedRef<class FMessagingDebuggerModel> FMessagingDebuggerModelRef;
-
-
 /**
  * Implements a view model for the messaging debugger.
  */
@@ -20,7 +13,7 @@ public:
 	/**
 	 * Clears all visibility filters (to show all messages).
 	 */
-	void ClearVisibilities( )
+	void ClearVisibilities()
 	{
 		InvisibleEndpoints.Empty();
 		InvisibleTypes.Empty();
@@ -33,7 +26,7 @@ public:
 	 *
 	 * @return The selected endpoint.
 	 */
-	FMessageTracerEndpointInfoPtr GetSelectedEndpoint( ) const
+	FMessageTracerEndpointInfoPtr GetSelectedEndpoint() const
 	{
 		return SelectedEndpoint;
 	}
@@ -43,7 +36,7 @@ public:
 	 *
 	 * @return The selected message.
 	 */
-	FMessageTracerMessageInfoPtr GetSelectedMessage( ) const
+	FMessageTracerMessageInfoPtr GetSelectedMessage() const
 	{
 		return SelectedMessage;
 	}
@@ -176,7 +169,7 @@ public:
 	 * @return The event delegate.
 	 */
 	DECLARE_EVENT(FMessagingDebuggerModel, FOnMessageVisibilityChanged);
-	FOnMessageVisibilityChanged& OnMessageVisibilityChanged( )
+	FOnMessageVisibilityChanged& OnMessageVisibilityChanged()
 	{
 		return MessageVisibilityChangedEvent;
 	}
@@ -187,7 +180,7 @@ public:
 	 * @return The event delegate.
 	 */
 	DECLARE_EVENT(FMessagingDebuggerModel, FOnSelectedEndpointChanged);
-	FOnSelectedEndpointChanged& OnSelectedEndpointChanged( )
+	FOnSelectedEndpointChanged& OnSelectedEndpointChanged()
 	{
 		return SelectedEndpointChangedEvent;
 	}
@@ -198,33 +191,40 @@ public:
 	 * @return The event delegate.
 	 */
 	DECLARE_EVENT(FMessagingDebuggerModel, FOnSelectedMessageChanged);
-	FOnSelectedMessageChanged& OnSelectedMessageChanged( )
+	FOnSelectedMessageChanged& OnSelectedMessageChanged()
 	{
 		return SelectedMessageChangedEvent;
 	}
 
 private:
 
-	// Holds the collection of invisible message endpoints.
+	/** Holds the collection of invisible message endpoints. */
 	TArray<FMessageTracerEndpointInfoPtr> InvisibleEndpoints;
 
-	// Holds the collection of invisible message types.
+	/** Holds the collection of invisible message types. */
 	TArray<FMessageTracerTypeInfoPtr> InvisibleTypes;
 
-	// Holds a pointer to the sole endpoint that is currently selected in the endpoint list.
+	/** Holds a pointer to the sole endpoint that is currently selected in the endpoint list. */
 	FMessageTracerEndpointInfoPtr SelectedEndpoint;
 
-	// Holds a pointer to the message that is currently selected in the message history.
+	/** Holds a pointer to the message that is currently selected in the message history. */
 	FMessageTracerMessageInfoPtr SelectedMessage;
 
 private:
 
-	// Holds an event delegate that is invoked when the visibility of messages has changed.
+	/** Holds an event delegate that is invoked when the visibility of messages has changed. */
 	FOnMessageVisibilityChanged MessageVisibilityChangedEvent;
 
-	// Holds an event delegate that is invoked when the selected endpoint has changed.
+	/** Holds an event delegate that is invoked when the selected endpoint has changed. */
 	FOnSelectedEndpointChanged SelectedEndpointChangedEvent;
 
-	// Holds an event delegate that is invoked when the selected message has changed.
+	/** Holds an event delegate that is invoked when the selected message has changed. */
 	FOnSelectedMessageChanged SelectedMessageChangedEvent;
 };
+
+
+/** Type definition for shared pointers to instances of FMessagingDebuggerModel. */
+typedef TSharedPtr<FMessagingDebuggerModel> FMessagingDebuggerModelPtr;
+
+/** Type definition for shared references to instances of FMessagingDebuggerTypeFilter. */
+typedef TSharedRef<FMessagingDebuggerModel> FMessagingDebuggerModelRef;
diff --git a/Engine/Plugins/Messaging/MessagingDebugger/Source/MessagingDebugger/Private/Models/MessagingDebuggerTypeFilter.h b/Engine/Plugins/Messaging/MessagingDebugger/Source/MessagingDebugger/Private/Models/MessagingDebuggerTypeFilter.h
--- a/Engine/Plugins/Messaging/MessagingDebugger/Source/MessagingDebugger/Private/Models/MessagingDebuggerTypeFilter.h
+++ b/Engine/Plugins/Messaging/MessagingDebugger/Source/MessagingDebugger/Private/Models/MessagingDebuggerTypeFilter.h
@@ -3,17 +3,6 @@
 #pragma once
 
 
-/**
- * Type definition for shared pointers to instances of FMessagingDebuggerTypeFilter.
- */
-typedef TSharedPtr<class FMessagingDebuggerTypeFilter> FMessagingDebuggerTypeFilterPtr;
-
-/**
- * Type definition for shared references to instances of FMessagingDebuggerTypeFilter.
- */
-typedef TSharedRef<class FMessagingDebuggerTypeFilter> FMessagingDebuggerTypeFilterRef;
-
-
 /**
  * Implements a filter for the message endpoints list.
  */
@@ -72,18 +61,25 @@ public:
 	 * @return The event delegate.
 	 */
 	DECLARE_EVENT(FMessagingDebuggerTypeFilter, FOnMessagingEndpointFilterChanged);
-	FOnMessagingEndpointFilterChanged& OnChanged( )
+	FOnMessagingEndpointFilterChanged& OnChanged()
 	{
 		return ChangedEvent;
 	}
 
 private:
 
-	// Holds the filter string used to filter message types by their names.
+	/** Holds the filter string used to filter message types by their names. */
 	FString FilterString;
 
 private:
 
-	// Holds an event delegate that is invoked when the filter settings changed.
+	/** Holds an event delegate that is invoked when the filter settings changed. */
 	FOnMessagingEndpointFilterChanged ChangedEvent;
 };
+
+
+/** Type definition for shared pointers to instances of FMessagingDebuggerTypeFilter. */
+typedef TSharedPtr<class FMessagingDebuggerTypeFilter> FMessagingDebuggerTypeFilterPtr;
+
+/** Type definition for shared references to instances of FMessagingDebuggerTypeFilter. */
+typedef TSharedRef<class FMessagingDebuggerTypeFilter> FMessagingDebuggerTypeFilterRef;
diff --git a/Engine/Plugins/Messaging/MessagingDebugger/Source/MessagingDebugger/Private/Widgets/Breakpoints/SMessagingBreakpoints.h b/Engine/Plugins/Messaging/MessagingDebugger/Source/MessagingDebugger/Private/Widgets/Breakpoints/SMessagingBreakpoints.h
--- a/Engine/Plugins/Messaging/MessagingDebugger/Source/MessagingDebugger/Private/Widgets/Breakpoints/SMessagingBreakpoints.h
+++ b/Engine/Plugins/Messaging/MessagingDebugger/Source/MessagingDebugger/Private/Widgets/Breakpoints/SMessagingBreakpoints.h
@@ -28,23 +28,23 @@ public:
 
 private:
 
-	// Handles generating a row widget for the endpoint list view.
+	/** Handles generating a row widget for the endpoint list view. */
 	TSharedRef<ITableRow> HandleBreakpointListGenerateRow( IMessageTracerBreakpointPtr EndpointInfo, const TSharedRef<STableViewBase>& OwnerTable );
 
-	// Handles the selection of endpoints.
+	/** Handles the selection of endpoints. */
 	void HandleBreakpointListSelectionChanged( IMessageTracerBreakpointPtr InItem, ESelectInfo::Type SelectInfo );
 
 private:
 
-	// Holds the filtered list of historic messages.
+	/** Holds the filtered list of historic messages. */
 	TArray<IMessageTracerBreakpointPtr> BreakpointList;
 
-	// Holds the message list view.
+	/** Holds the message list view. */
 	TSharedPtr<SListView<IMessageTracerBreakpointPtr>> BreakpointListView;
 
-	// Holds the widget's visual style.
+	/** Holds the widget's visual style. */
 	TSharedPtr<ISlateStyle> Style;
 
-	// Holds a pointer to the message bus tracer.
+	/** Holds a pointer to the message bus tracer. */
 	IMessageTracerPtr Tracer;
 };
diff --git a/Engine/Plugins/Messaging/MessagingDebugger/Source/MessagingDebugger/Private/Widgets/Breakpoints/SMessagingBreakpointsTableRow.h b/Engine/Plugins/Messaging/MessagingDebugger/Source/MessagingDebugger/Private/Widgets/Breakpoints/SMessagingBreakpointsTableRow.h
--- a/Engine/Plugins/Messaging/MessagingDebugger/Source/MessagingDebugger/Private/Widgets/Breakpoints/SMessagingBreakpointsTableRow.h
+++ b/Engine/Plugins/Messaging/MessagingDebugger/Source/MessagingDebugger/Private/Widgets/Breakpoints/SMessagingBreakpointsTableRow.h
@@ -80,10 +80,10 @@ public:
 
 private:
 
-	// Holds a pointer to the breakpoint that is shown in this row.
+	/** Holds a pointer to the breakpoint that is shown in this row. */
 	IMessageTracerBreakpointPtr Breakpoint;
 
-	// Holds the widget's visual style.
+	/** Holds the widget's visual style. */
 	TSharedPtr<ISlateStyle> Style;
 };
 
diff --git a/Engine/Plugins/Messaging/MessagingDebugger/Source/MessagingDebugger/Private/Widgets/EndpointDetails/SMessagingAddressTableRow.h b/Engine/Plugins/Messaging/MessagingDebugger/Source/MessagingDebugger/Private/Widgets/EndpointDetails/SMessagingAddressTableRow.h
--- a/Engine/Plugins/Messaging/MessagingDebugger/Source/MessagingDebugger/Private/Widgets/EndpointDetails/SMessagingAddressTableRow.h
+++ b/Engine/Plugins/Messaging/MessagingDebugger/Source/MessagingDebugger/Private/Widgets/EndpointDetails/SMessagingAddressTableRow.h
@@ -24,8 +24,8 @@ public:
 	/**
 	 * Constructs the widget.
 	 *
-	 * @param InArgs - The construction arguments.
-	 * @param InOwnerTableView - The table view that owns this row.
+	 * @param InArgs The construction arguments.
+	 * @param InOwnerTableView The table view that owns this row.
 	 */
 	void Construct( const FArguments& InArgs, const TSharedRef<STableViewBase>& InOwnerTableView, const FMessagingDebuggerModelRef& InModel )
 	{
@@ -83,13 +83,13 @@ public:
 
 private:
 
-	// Holds the address information.
+	/** Holds the address information. */
 	FMessageTracerAddressInfoPtr AddressInfo;
 
-	// Holds a pointer to the view model.
+	/** Holds a pointer to the view model. */
 	FMessagingDebuggerModelPtr Model;
 
-	// Holds the widget's visual style.
+	/** Holds the widget's visual style. */
 	TSharedPtr<ISlateStyle> Style;
 };
 
diff --git a/Engine/Plugins/Messaging/MessagingDebugger/Source/MessagingDebugger/Private/Widgets/EndpointDetails/SMessagingEndpointDetails.h b/Engine/Plugins/Messaging/MessagingDebugger/Source/MessagingDebugger/Private/Widgets/EndpointDetails/SMessagingEndpointDetails.h
--- a/Engine/Plugins/Messaging/MessagingDebugger/Source/MessagingDebugger/Private/Widgets/EndpointDetails/SMessagingEndpointDetails.h
+++ b/Engine/Plugins/Messaging/MessagingDebugger/Source/MessagingDebugger/Private/Widgets/EndpointDetails/SMessagingEndpointDetails.h
@@ -33,33 +33,31 @@ public:
 
 protected:
 
-	/**
-	 * Refreshes the endpoint's address information.
-	 */
-	void RefreshAddressInfo( );
+	/** Refreshes the endpoint's address information. */
+	void RefreshAddressInfo();
 
 private:
 
-	// Callback for generating a row widget for the address list view.
+	/** Callback for generating a row widget for the address list view. */
 	TSharedRef<ITableRow> HandleAddressListGenerateRow( FMessageTracerAddressInfoPtr AddressInfo, const TSharedRef<STableViewBase>& OwnerTable );
 
-	// Callback for getting the number of received messages.
-	FString HandleEndpointDetailsReceivedMessagesText( ) const;
+	/** Callback for getting the number of received messages. */
+	FString HandleEndpointDetailsReceivedMessagesText() const;
 
-	// Callback for getting the number of sent messages.
-	FString HandleEndpointDetailsSentMessagesText( ) const;
+	/** Callback for getting the number of sent messages. */
+	FString HandleEndpointDetailsSentMessagesText() const;
 
 private:
 
-	// Holds the list of address information.
+	/** Holds the list of address information. */
 	TArray<FMessageTracerAddressInfoPtr> AddressList;
 
-	// Holds the address information list view.
+	/** Holds the address information list view. */
 	TSharedPtr<SListView<FMessageTracerAddressInfoPtr> > AddressListView;
 
-	// Holds a pointer to the view model.
+	/** Holds a pointer to the view model. */
 	FMessagingDebuggerModelPtr Model;
 
-	// Holds the widget's visual style.
+	/** Holds the widget's visual style. */
 	TSharedPtr<ISlateStyle> Style;
 };
diff --git a/Engine/Plugins/Messaging/MessagingDebugger/Source/MessagingDebugger/Private/Widgets/Endpoints/SMessagingEndpoints.h b/Engine/Plugins/Messaging/MessagingDebugger/Source/MessagingDebugger/Private/Widgets/Endpoints/SMessagingEndpoints.h
--- a/Engine/Plugins/Messaging/MessagingDebugger/Source/MessagingDebugger/Private/Widgets/Endpoints/SMessagingEndpoints.h
+++ b/Engine/Plugins/Messaging/MessagingDebugger/Source/MessagingDebugger/Private/Widgets/Endpoints/SMessagingEndpoints.h
@@ -16,9 +16,7 @@ public:
 
 public:
 
-	/**
-	 * Destructor.
-	 */
+	/**  Destructor. */
 	~SMessagingEndpoints();
 
 public:
@@ -41,48 +39,46 @@ public:
 
 protected:
 
-	/**
-	 * Reloads the list of endpoints.
-	 */
-	void ReloadEndpointList( );
+	/** Reloads the list of endpoints. */
+	void ReloadEndpointList();
 
 private:
 
-	// Handles generating a row widget for the endpoint list view.
+	/** Handles generating a row widget for the endpoint list view. */
 	TSharedRef<ITableRow> HandleEndpointListGenerateRow( FMessageTracerEndpointInfoPtr EndpointInfo, const TSharedRef<STableViewBase>& OwnerTable );
 
-	// Handles getting the highlight string for endpoints.
-	FText HandleEndpointListGetHighlightText( ) const;
+	/** Handles getting the highlight string for endpoints. */
+	FText HandleEndpointListGetHighlightText() const;
 
-	// Handles the selection of endpoints.
+	/** Handles the selection of endpoints. */
 	void HandleEndpointListSelectionChanged( FMessageTracerEndpointInfoPtr InItem, ESelectInfo::Type SelectInfo );
 
-	// Handles endpoint filter changes.
-	void HandleFilterChanged( );
+	/** Handles endpoint filter changes. */
+	void HandleFilterChanged();
 
-	// Callback for handling message selection changes.
-	void HandleModelSelectedMessageChanged( );
+	/** Callback for handling message selection changes. */
+	void HandleModelSelectedMessageChanged();
 
 private:
 
-	// Holds the filter bar.
-	//TSharedPtr<SMessagingEndpointsFilterBar> FilterBar;
+	/** Holds the filter bar. */
+	/**TSharedPtr<SMessagingEndpointsFilterBar> FilterBar;
 
-	// Holds the filtered list of historic messages.
+	/** Holds the filtered list of historic messages. */
 	TArray<FMessageTracerEndpointInfoPtr> EndpointList;
 
-	// Holds the message list view.
+	/** Holds the message list view. */
 	TSharedPtr<SListView<FMessageTracerEndpointInfoPtr>> EndpointListView;
 
-	// Holds the endpoint filter model.
+	/** Holds the endpoint filter model. */
 	FMessagingDebuggerEndpointFilterPtr Filter;
 
-	// Holds a pointer to the view model.
+	/** Holds a pointer to the view model. */
 	FMessagingDebuggerModelPtr Model;
 
-	// Holds the widget's visual style.
+	/** Holds the widget's visual style. */
 	TSharedPtr<ISlateStyle> Style;
 
-	// Holds a pointer to the message bus tracer.
+	/** Holds a pointer to the message bus tracer. */
 	IMessageTracerPtr Tracer;
 };
diff --git a/Engine/Plugins/Messaging/MessagingDebugger/Source/MessagingDebugger/Private/Widgets/Endpoints/SMessagingEndpointsFilterBar.cpp b/Engine/Plugins/Messaging/MessagingDebugger/Source/MessagingDebugger/Private/Widgets/Endpoints/SMessagingEndpointsFilterBar.cpp
--- a/Engine/Plugins/Messaging/MessagingDebugger/Source/MessagingDebugger/Private/Widgets/Endpoints/SMessagingEndpointsFilterBar.cpp
+++ b/Engine/Plugins/Messaging/MessagingDebugger/Source/MessagingDebugger/Private/Widgets/Endpoints/SMessagingEndpointsFilterBar.cpp
@@ -1,11 +1,16 @@
 // Copyright 1998-2014 Epic Games, Inc. All Rights Reserved.
-#include "MessagingDebuggerPrivatePCH.h"
 
+#include "MessagingDebuggerPrivatePCH.h"
 #include "SMessagingEndpointsFilterBar.h"
 #include "SSearchBox.h"
 
+
 #define LOCTEXT_NAMESPACE "SMessagingEndpointsFilterBar"
 
+
+/* SMessagingEndpointsFilterBar interface
+*****************************************************************************/
+
 void SMessagingEndpointsFilterBar::Construct(const FArguments& InArgs, FMessagingDebuggerEndpointFilterRef InFilter)
 {
 	Filter = InFilter;
@@ -26,4 +31,5 @@ void SMessagingEndpointsFilterBar::Construct(const FArguments& InArgs, FMessagin
 		];
 }
 
+
 #undef LOCTEXT_NAMESPACE
diff --git a/Engine/Plugins/Messaging/MessagingDebugger/Source/MessagingDebugger/Private/Widgets/Endpoints/SMessagingEndpointsFilterBar.h b/Engine/Plugins/Messaging/MessagingDebugger/Source/MessagingDebugger/Private/Widgets/Endpoints/SMessagingEndpointsFilterBar.h
--- a/Engine/Plugins/Messaging/MessagingDebugger/Source/MessagingDebugger/Private/Widgets/Endpoints/SMessagingEndpointsFilterBar.h
+++ b/Engine/Plugins/Messaging/MessagingDebugger/Source/MessagingDebugger/Private/Widgets/Endpoints/SMessagingEndpointsFilterBar.h
@@ -26,7 +26,7 @@ public:
 
 private:
 
-	// Handles changing the filter string text box text.
+	/** Handles changing the filter string text box text. */
 	void HandleFilterStringTextChanged( const FText& NewText )
 	{
 		Filter->SetFilterString(NewText.ToString());
@@ -34,6 +34,6 @@ private:
 
 private:
 
-	// Holds a pointer to the filter model.
+	/** Holds a pointer to the filter model. */
 	FMessagingDebuggerEndpointFilterPtr Filter;
 };
diff --git a/Engine/Plugins/Messaging/MessagingDebugger/Source/MessagingDebugger/Private/Widgets/Endpoints/SMessagingEndpointsTableRow.h b/Engine/Plugins/Messaging/MessagingDebugger/Source/MessagingDebugger/Private/Widgets/Endpoints/SMessagingEndpointsTableRow.h
--- a/Engine/Plugins/Messaging/MessagingDebugger/Source/MessagingDebugger/Private/Widgets/Endpoints/SMessagingEndpointsTableRow.h
+++ b/Engine/Plugins/Messaging/MessagingDebugger/Source/MessagingDebugger/Private/Widgets/Endpoints/SMessagingEndpointsTableRow.h
@@ -112,26 +112,26 @@ public:
 
 private:
 
-	// Gets the text for the Messages column.
-	FString HandleMessagesText( ) const
+	/** Gets the text for the Messages column. */
+	FString HandleMessagesText() const
 	{
 		return FString::Printf(TEXT("%i"), EndpointInfo->ReceivedMessages.Num() + EndpointInfo->SentMessages.Num());
 	}
 
-	// Gets the tooltip text for the Messages column.
-	FString HandleMessagesTooltipText( ) const
+	/** Gets the tooltip text for the Messages column. */
+	FString HandleMessagesTooltipText() const
 	{
 		return FString::Printf(TEXT("In: %i\nOut: %i"), EndpointInfo->ReceivedMessages.Num(), EndpointInfo->SentMessages.Num());
 	}
 
-	// Handles changing the checked state of the visibility check box.
+	/** Handles changing the checked state of the visibility check box. */
 	void HandleVisibilityCheckBoxCheckStateChanged( ESlateCheckBoxState::Type CheckState )
 	{
 		Model->SetEndpointVisibility(EndpointInfo.ToSharedRef(), (CheckState == ESlateCheckBoxState::Checked));
 	}
 
-	// Gets the image for the visibility check box.
-	ESlateCheckBoxState::Type HandleVisibilityCheckBoxIsChecked( ) const
+	/** Gets the image for the visibility check box. */
+	ESlateCheckBoxState::Type HandleVisibilityCheckBoxIsChecked() const
 	{
 		if (Model->IsEndpointVisible(EndpointInfo.ToSharedRef()))
 		{
@@ -143,16 +143,16 @@ private:
 
 private:
 
-	// Holds the endpoint's debug information.
+	/** Holds the endpoint's debug information. */
 	FMessageTracerEndpointInfoPtr EndpointInfo;
 
-	// Holds the highlight string for the message.
+	/** Holds the highlight string for the message. */
 	TAttribute<FText> HighlightText;
 
-	// Holds a pointer to the view model.
+	/** Holds a pointer to the view model. */
 	FMessagingDebuggerModelPtr Model;
 
-	// Holds the widget's visual style.
+	/** Holds the widget's visual style. */
 	TSharedPtr<ISlateStyle> Style;
 };
 
diff --git a/Engine/Plugins/Messaging/MessagingDebugger/Source/MessagingDebugger/Private/Widgets/Graph/SMessagingGraph.h b/Engine/Plugins/Messaging/MessagingDebugger/Source/MessagingDebugger/Private/Widgets/Graph/SMessagingGraph.h
--- a/Engine/Plugins/Messaging/MessagingDebugger/Source/MessagingDebugger/Private/Widgets/Graph/SMessagingGraph.h
+++ b/Engine/Plugins/Messaging/MessagingDebugger/Source/MessagingDebugger/Private/Widgets/Graph/SMessagingGraph.h
@@ -26,6 +26,6 @@ public:
 
 private:
 
-	// Holds the graph editor widget.
+	/** Holds the graph editor widget. */
 	//TSharedPtr<SGraphEditor> GraphEditor;
 };
diff --git a/Engine/Plugins/Messaging/MessagingDebugger/Source/MessagingDebugger/Private/Widgets/History/SMessagingHistory.h b/Engine/Plugins/Messaging/MessagingDebugger/Source/MessagingDebugger/Private/Widgets/History/SMessagingHistory.h
--- a/Engine/Plugins/Messaging/MessagingDebugger/Source/MessagingDebugger/Private/Widgets/History/SMessagingHistory.h
+++ b/Engine/Plugins/Messaging/MessagingDebugger/Source/MessagingDebugger/Private/Widgets/History/SMessagingHistory.h
@@ -16,10 +16,8 @@ public:
 
 public:
 
-	/**
-	 * Destructor.
-	 */
-	~SMessagingHistory( );
+	/** Destructor. */
+	~SMessagingHistory();
 
 public:
 
@@ -42,72 +40,70 @@ protected:
 	 */
 	void AddMessage( const FMessageTracerMessageInfoRef& MessageInfo );
 
-	/**
-	 * Reloads the message history.
-	 */
-	void ReloadMessages( );
+	/** Reloads the message history. */
+	void ReloadMessages();
 
 private:
 
-	// Callback for message filter changes.
-	void HandleFilterChanged( );
+	/** Callback for message filter changes. */
+	void HandleFilterChanged();
 
-	// Callback for generating a row widget for the message list view.
+	/** Callback for generating a row widget for the message list view. */
 	TSharedRef<ITableRow> HandleMessageListGenerateRow( FMessageTracerMessageInfoPtr MessageInfo, const TSharedRef<STableViewBase>& OwnerTable );
 
-	// Callback for getting the message list's highlight string.
-	FText HandleMessageListGetHighlightText( ) const;
+	/** Callback for getting the message list's highlight string. */
+	FText HandleMessageListGetHighlightText() const;
 
-	// Callback for double-clicking a message item.
+	/** Callback for double-clicking a message item. */
 	void HandleMessageListItemDoubleClick( FMessageTracerMessageInfoPtr Item );
 
-	// Callback for scrolling a message item into view.
+	/** Callback for scrolling a message item into view. */
 	void HandleMessageListItemScrolledIntoView( FMessageTracerMessageInfoPtr Item, const TSharedPtr<ITableRow>& TableRow );
 
-	// Callback for selecting a message in the message list view.
+	/** Callback for selecting a message in the message list view. */
 	void HandleMessageListSelectionChanged( FMessageTracerMessageInfoPtr InItem, ESelectInfo::Type SelectInfo );
 
-	// Callback for changes in message visibility.
-	void HandleModelMessageVisibilityChanged( );
+	/** Callback for changes in message visibility. */
+	void HandleModelMessageVisibilityChanged();
 
-	// Callback for clicking the 'Show hidden' hyperlink.
-	void HandleShowHiddenHyperlinkNavigate( );
+	/** Callback for clicking the 'Show hidden' hyperlink. */
+	void HandleShowHiddenHyperlinkNavigate();
 
-	// Callback for getting the visibility of the 'Show hidden' hyperlink.
-	EVisibility HandleShowHiddenHyperlinkVisibility( ) const;
+	/** Callback for getting the visibility of the 'Show hidden' hyperlink. */
+	EVisibility HandleShowHiddenHyperlinkVisibility() const;
 
-	// Callback for getting the status bar text.
-	FText HandleStatusBarText( ) const;
+	/** Callback for getting the status bar text. */
+	FText HandleStatusBarText() const;
 
-	// Callback for when a message has been added to the message tracer.
+	/** Callback for when a message has been added to the message tracer. */
 	void HandleTracerMessageAdded( FMessageTracerMessageInfoRef MessageInfo );
 
-	// Callback for when the tracer's message history has been reset.
-	void HandleTracerMessagesReset( );
+	/** Callback for when the tracer's message history has been reset. */
+	void HandleTracerMessagesReset();
 
 private:
 
-	// Holds the message filter model.
+	/** Holds the message filter model. */
 	FMessagingDebuggerMessageFilterPtr Filter;
 
-	// Holds the filtered list of historic messages.
+	/** Holds the filtered list of historic messages. */
 	TArray<FMessageTracerMessageInfoPtr> MessageList;
 
-	// Holds the message list view.
+	/** Holds the message list view. */
 	TSharedPtr<SListView<FMessageTracerMessageInfoPtr>> MessageListView;
 
-	// Holds a pointer to the view model.
+	/** Holds a pointer to the view model. */
 	FMessagingDebuggerModelPtr Model;
 
-	// Holds a flag indicating whether the message list should auto-scroll to the last item.
+	/** Holds a flag indicating whether the message list should auto-scroll to the last item. */
 	bool ShouldScrollToLast;
 
-	// Holds the widget's visual style.
+	/** Holds the widget's visual style. */
 	TSharedPtr<ISlateStyle> Style;
 
-	// Holds the total number of messages (including invisible ones).
+	/** Holds the total number of messages (including invisible ones). */
 	int32 TotalMessages;
 
-	// Holds a pointer to the message bus tracer.
+	/** Holds a pointer to the message bus tracer. */
 	IMessageTracerPtr Tracer;
 };
diff --git a/Engine/Plugins/Messaging/MessagingDebugger/Source/MessagingDebugger/Private/Widgets/History/SMessagingHistoryFilterBar.h b/Engine/Plugins/Messaging/MessagingDebugger/Source/MessagingDebugger/Private/Widgets/History/SMessagingHistoryFilterBar.h
--- a/Engine/Plugins/Messaging/MessagingDebugger/Source/MessagingDebugger/Private/Widgets/History/SMessagingHistoryFilterBar.h
+++ b/Engine/Plugins/Messaging/MessagingDebugger/Source/MessagingDebugger/Private/Widgets/History/SMessagingHistoryFilterBar.h
@@ -35,7 +35,7 @@ public:
 
 private:
 
-	// Holds the filter model.
+	/** Holds the filter model. */
 	FMessagingDebuggerMessageFilterPtr Filter;
 
 };
diff --git a/Engine/Plugins/Messaging/MessagingDebugger/Source/MessagingDebugger/Private/Widgets/History/SMessagingHistoryTableRow.h b/Engine/Plugins/Messaging/MessagingDebugger/Source/MessagingDebugger/Private/Widgets/History/SMessagingHistoryTableRow.h
--- a/Engine/Plugins/Messaging/MessagingDebugger/Source/MessagingDebugger/Private/Widgets/History/SMessagingHistoryTableRow.h
+++ b/Engine/Plugins/Messaging/MessagingDebugger/Source/MessagingDebugger/Private/Widgets/History/SMessagingHistoryTableRow.h
@@ -287,8 +287,8 @@ protected:
 
 private:
 
-	// Callback for getting the text color of the DispatchLatency column.
-	FSlateColor HandleDispatchLatencyColorAndOpacity( ) const
+	/** Callback for getting the text color of the DispatchLatency column. */
+	FSlateColor HandleDispatchLatencyColorAndOpacity() const
 	{
 		if (MaxDispatchLatency >= 0.01)
 		{
@@ -308,13 +308,13 @@ private:
 		return FSlateColor::UseForeground();
 	}
 
-	// Callback for getting the text in the DispatchLatency column.
-	FText HandleDispatchLatencyText( ) const
+	/** Callback for getting the text in the DispatchLatency column. */
+	FText HandleDispatchLatencyText() const
 	{
 		return TimespanToReadableText(MaxDispatchLatency);
 	}
 
-	const FSlateBrush* HandleFlagImage( ) const
+	const FSlateBrush* HandleFlagImage() const
 	{
 		if ((MessageInfo->TimeRouted > 0.0) && (MessageInfo->DispatchStates.Num() == 0))
 		{
@@ -324,14 +324,14 @@ private:
 		return NULL;
 	}
 
-	// Callback for getting the text in the HandlingTime column.
-	FText HandleHandlingTimeText( ) const
+	/** Callback for getting the text in the HandlingTime column. */
+	FText HandleHandlingTimeText() const
 	{
 		return TimespanToReadableText(MaxHandlingTime);
 	}
 
-	// Callback for getting the text in the Recipients column.
-	FText HandleRecipientsText( ) const
+	/** Callback for getting the text in the Recipients column. */
+	FText HandleRecipientsText() const
 	{
 		int32 LocalRecipients = 0;
 		int32 RemoteRecipients = 0;
@@ -356,8 +356,8 @@ private:
 		return FText::Format(LOCTEXT("RecipientsTextFormat", "{0} / {1}"), FText::AsNumber(LocalRecipients), FText::AsNumber(RemoteRecipients));
 	}
 
-	// Callback for getting the text color of the RouteLatency column.
-	FSlateColor HandleRouteLatencyColorAndOpacity( ) const
+	/** Callback for getting the text color of the RouteLatency column. */
+	FSlateColor HandleRouteLatencyColorAndOpacity() const
 	{
 		double RouteLatency = MessageInfo->TimeRouted - MessageInfo->TimeSent;
 
@@ -384,8 +384,8 @@ private:
 		return FSlateColor::UseForeground();
 	}
 
-	// Callback for getting the text in the RouteLatency column.
-	FText HandleRouteLatencyText( ) const
+	/** Callback for getting the text in the RouteLatency column. */
+	FText HandleRouteLatencyText() const
 	{
 		if (MessageInfo->TimeRouted > 0.0)
 		{
@@ -395,14 +395,14 @@ private:
 		return LOCTEXT("RouteLatencyPending", "Pending");
 	}
 
-	// Callback for getting the text in the Sender column.
-	FText HandleSenderText( ) const
+	/** Callback for getting the text in the Sender column. */
+	FText HandleSenderText() const
 	{
 		return FText::FromString(MessageInfo->SenderInfo->Name.ToString());
 	}
 
-	// Callback for getting the text color of various columns.
-	FSlateColor HandleTextColorAndOpacity( ) const
+	/** Callback for getting the text color of various columns. */
+	FSlateColor HandleTextColorAndOpacity() const
 	{
 		if (MessageInfo->TimeRouted == 0.0)
 		{
@@ -414,19 +414,19 @@ private:
 
 private:
 
-	// Holds the highlight string for the message.
+	/** Holds the highlight string for the message. */
 	TAttribute<FText> HighlightText;
 
-	// Holds message's debug information.
+	/** Holds message's debug information. */
 	FMessageTracerMessageInfoPtr MessageInfo;
 
-	// Holds the maximum dispatch latency.
+	/** Holds the maximum dispatch latency. */
 	double MaxDispatchLatency;
 
-	// Holds the maximum time that was needed to handle the message.
+	/** Holds the maximum time that was needed to handle the message. */
 	double MaxHandlingTime;
 
-	// Holds the widget's visual style.
+	/** Holds the widget's visual style. */
 	TSharedPtr<ISlateStyle> Style;
 };
 
diff --git a/Engine/Plugins/Messaging/MessagingDebugger/Source/MessagingDebugger/Private/Widgets/MessageDetails/SMessagingDispatchStateTableRow.h b/Engine/Plugins/Messaging/MessagingDebugger/Source/MessagingDebugger/Private/Widgets/MessageDetails/SMessagingDispatchStateTableRow.h
--- a/Engine/Plugins/Messaging/MessagingDebugger/Source/MessagingDebugger/Private/Widgets/MessageDetails/SMessagingDispatchStateTableRow.h
+++ b/Engine/Plugins/Messaging/MessagingDebugger/Source/MessagingDebugger/Private/Widgets/MessageDetails/SMessagingDispatchStateTableRow.h
@@ -159,8 +159,8 @@ protected:
 
 private:
 
-	// Callback for getting the handling time text.
-	FText HandleHandlingTimeText( ) const
+	/** Callback for getting the handling time text. */
+	FText HandleHandlingTimeText() const
 	{
 		if (DispatchState->TimeHandled > 0.0)
 		{
@@ -170,8 +170,8 @@ private:
 		return LOCTEXT("NotHandledYetText", "Not handled yet");
 	}
 
-	// Callback for getting the color of a time span text.
-	FSlateColor HandleHandlingTimeColorAndOpacity( ) const
+	/** Callback for getting the color of a time span text. */
+	FSlateColor HandleHandlingTimeColorAndOpacity() const
 	{
 		if (DispatchState->TimeHandled == 0.0)
 		{
@@ -181,8 +181,8 @@ private:
 		return FSlateColor::UseForeground();
 	}
 
-	// Callback for getting the dispatch type image.
-	const FSlateBrush* HandleTypeImage( ) const
+	/** Callback for getting the dispatch type image. */
+	const FSlateBrush* HandleTypeImage() const
 	{
 		if (DispatchState->DispatchType == EMessageTracerDispatchTypes::Direct)
 		{
@@ -197,8 +197,8 @@ private:
 		return Style->GetBrush("DispatchPending");
 	}
 
-	// Callback for getting the dispatch type tool tip text.
-	FText HandleTypeTooltip( ) const
+	/** Callback for getting the dispatch type tool tip text. */
+	FText HandleTypeTooltip() const
 	{
 		if (DispatchState->DispatchType == EMessageTracerDispatchTypes::Direct)
 		{
@@ -215,13 +215,13 @@ private:
 
 private:
 
-	// Holds the message dispatch state.
+	/** Holds the message dispatch state. */
 	FMessageTracerDispatchStatePtr DispatchState;
 
-	// Holds a pointer to the view model.
+	/** Holds a pointer to the view model. */
 	FMessagingDebuggerModelPtr Model;
 
-	// Holds the widget's visual style.
+	/** Holds the widget's visual style. */
 	TSharedPtr<ISlateStyle> Style;
 };
 
diff --git a/Engine/Plugins/Messaging/MessagingDebugger/Source/MessagingDebugger/Private/Widgets/MessageDetails/SMessagingMessageDetails.h b/Engine/Plugins/Messaging/MessagingDebugger/Source/MessagingDebugger/Private/Widgets/MessageDetails/SMessagingMessageDetails.h
--- a/Engine/Plugins/Messaging/MessagingDebugger/Source/MessagingDebugger/Private/Widgets/MessageDetails/SMessagingMessageDetails.h
+++ b/Engine/Plugins/Messaging/MessagingDebugger/Source/MessagingDebugger/Private/Widgets/MessageDetails/SMessagingMessageDetails.h
@@ -16,10 +16,8 @@ public:
 
 public:
 
-	/**
-	 * Destructor.
-	 */
-	~SMessagingMessageDetails( );
+	/** Destructor. */
+	~SMessagingMessageDetails();
 
 public:
 
@@ -34,45 +32,43 @@ public:
 
 protected:
 
-	/**
-	 * Refreshes the details widget.
-	 */
-	void RefreshDetails( );
+	/** Refreshes the details widget. */
+	void RefreshDetails();
 
 private:
 
-	// Callback for generating a row widget for the dispatch state list view.
+	/** Callback for generating a row widget for the dispatch state list view. */
 	TSharedRef<ITableRow> HandleDispatchStateListGenerateRow( FMessageTracerDispatchStatePtr DispatchState, const TSharedRef<STableViewBase>& OwnerTable );
 
-	// Callback for getting the text of the 'Expiration' field.
-	FText HandleExpirationText( ) const;
+	/** Callback for getting the text of the 'Expiration' field. */
+	FText HandleExpirationText() const;
 
-	// Callback for getting the text of the 'Message Type' field.
-	FText HandleMessageTypeText( ) const;
+	/** Callback for getting the text of the 'Message Type' field. */
+	FText HandleMessageTypeText() const;
 
-	// Callback for handling message selection changes.
-	void HandleModelSelectedMessageChanged( );
+	/** Callback for handling message selection changes. */
+	void HandleModelSelectedMessageChanged();
 
-	// Callback for getting the text of the 'Sender' field.
-	FText HandleSenderText( ) const;
+	/** Callback for getting the text of the 'Sender' field. */
+	FText HandleSenderText() const;
 
-	// Callback for getting the text of the 'Sender Thread' field.
+	/** Callback for getting the text of the 'Sender Thread' field. */
 	FText HandleSenderThreadText() const;
 
-	// Callback for getting the text of the 'Timestamp' field.
+	/** Callback for getting the text of the 'Timestamp' field. */
 	FText HandleTimestampText() const;
 
 private:
 
-	// Holds the list of dispatch states.
+	/** Holds the list of dispatch states. */
 	TArray<FMessageTracerDispatchStatePtr> DispatchStateList;
 
-	// Holds the dispatch state list view.
+	/** Holds the dispatch state list view. */
 	TSharedPtr<SListView<FMessageTracerDispatchStatePtr> > DispatchStateListView;
 
-	// Holds a pointer to the view model.
+	/** Holds a pointer to the view model. */
 	FMessagingDebuggerModelPtr Model;
 
-	// Holds the widget's visual style.
+	/** Holds the widget's visual style. */
 	TSharedPtr<ISlateStyle> Style;
 };
diff --git a/Engine/Plugins/Messaging/MessagingDebugger/Source/MessagingDebugger/Private/Widgets/SMessagingDebugger.h b/Engine/Plugins/Messaging/MessagingDebugger/Source/MessagingDebugger/Private/Widgets/SMessagingDebugger.h
--- a/Engine/Plugins/Messaging/MessagingDebugger/Source/MessagingDebugger/Private/Widgets/SMessagingDebugger.h
+++ b/Engine/Plugins/Messaging/MessagingDebugger/Source/MessagingDebugger/Private/Widgets/SMessagingDebugger.h
@@ -16,10 +16,8 @@ public:
 
 public:
 
-	/**
-	 * Default constructor.
-	 */
-	SMessagingDebugger( );
+	/** Default constructor. */
+	SMessagingDebugger();
 
 public:
 
@@ -51,68 +49,68 @@ protected:
 
 private:
 
-	// Callback for checking whether the BreakDebugger command can be executed.
-	bool HandleBreakDebuggerCommandCanExecute( ) const;
+	/** Callback for checking whether the BreakDebugger command can be executed. */
+	bool HandleBreakDebuggerCommandCanExecute() const;
 
-	// Callback for clicking the BreakDebugger command.
-	void HandleBreakDebuggerCommandExecute( );
+	/** Callback for clicking the BreakDebugger command. */
+	void HandleBreakDebuggerCommandExecute();
 
-	// Callback for checking whether the ClearHistory command can be executed.
-	bool HandleClearHistoryCommandCanExecute( ) const;
+	/** Callback for checking whether the ClearHistory command can be executed. */
+	bool HandleClearHistoryCommandCanExecute() const;
 
-	// Callback for clicking the ClearHistory command.
-	void HandleClearHistoryCommandExecute( );
+	/** Callback for clicking the ClearHistory command. */
+	void HandleClearHistoryCommandExecute();
 
-	// Callback for checking whether the ContinueDebugger command can be executed.
-	bool HandleContinueDebuggerCommandCanExecute( ) const;
+	/** Callback for checking whether the ContinueDebugger command can be executed. */
+	bool HandleContinueDebuggerCommandCanExecute() const;
 
-	// Callback for clicking the ContinueDebugger command.
-	void HandleContinueDebuggerCommandExecute( );
+	/** Callback for clicking the ContinueDebugger command. */
+	void HandleContinueDebuggerCommandExecute();
 
-	// Callback for checking whether the ContinueDebugger command is visible.
-	bool HandleContinueDebuggerCommandIsVisible( ) const;
+	/** Callback for checking whether the ContinueDebugger command is visible. */
+	bool HandleContinueDebuggerCommandIsVisible() const;
 
-	// Callback for for when the owner tab's visual state is being persisted.
-	void HandleMajorTabPersistVisualState( );
+	/** Callback for for when the owner tab's visual state is being persisted. */
+	void HandleMajorTabPersistVisualState();
 
-	// Callback for checking whether the StartDebugger command can be executed.
-	bool HandleStartDebuggerCommandCanExecute( ) const;
+	/** Callback for checking whether the StartDebugger command can be executed. */
+	bool HandleStartDebuggerCommandCanExecute() const;
 
-	// Callback for clicking the StartDebugger command.
-	void HandleStartDebuggerCommandExecute( );
+	/** Callback for clicking the StartDebugger command. */
+	void HandleStartDebuggerCommandExecute();
 
-	// Callback for checking whether the StartDebugger command is visible.
-	bool HandleStartDebuggerCommandIsVisible( ) const;
+	/** Callback for checking whether the StartDebugger command is visible. */
+	bool HandleStartDebuggerCommandIsVisible() const;
 
-	// Callback for checking whether the StepDebugger command can be executed.
-	bool HandleStepDebuggerCommandCanExecute( ) const;
+	/** Callback for checking whether the StepDebugger command can be executed. */
+	bool HandleStepDebuggerCommandCanExecute() const;
 
-	// Callback for clicking the StepDebugger command.
-	void HandleStepDebuggerCommandExecute( );
+	/** Callback for clicking the StepDebugger command. */
+	void HandleStepDebuggerCommandExecute();
 
-	// Callback for checking whether the StopDebugger command can be executed.
-	bool HandleStopDebuggerCommandCanExecute( ) const;
+	/** Callback for checking whether the StopDebugger command can be executed. */
+	bool HandleStopDebuggerCommandCanExecute() const;
 
-	// Callback for clicking the StopDebugger command.
-	void HandleStopDebuggerCommandExecute( );
+	/** Callback for clicking the StopDebugger command. */
+	void HandleStopDebuggerCommandExecute();
 
-	// Callback for spawning tabs.
+	/** Callback for spawning tabs. */
 	TSharedRef<SDockTab> HandleTabManagerSpawnTab( const FSpawnTabArgs& Args, FName TabIdentifier ) const;
 
 private:
 
-	// Holds the list of UI commands.
+	/** Holds the list of UI commands. */
 	TSharedRef<FUICommandList> CommandList;
 
-	// Holds a pointer to the message bus tracer.
+	/** Holds a pointer to the message bus tracer. */
 	IMessageTracerPtr MessageTracer;
 
-	// Holds the messaging debugger's view model.
+	/** Holds the messaging debugger's view model. */
 	FMessagingDebuggerModelRef Model;
 
-	// Holds the widget's visual style.
+	/** Holds the widget's visual style. */
 	TSharedPtr<ISlateStyle> Style;
 
-	// Holds the tab manager that manages the front-end's tabs.
+	/** Holds the tab manager that manages the front-end's tabs. */
 	TSharedPtr<FTabManager> TabManager;
 };
diff --git a/Engine/Plugins/Messaging/MessagingDebugger/Source/MessagingDebugger/Private/Widgets/Types/SMessagingTypes.h b/Engine/Plugins/Messaging/MessagingDebugger/Source/MessagingDebugger/Private/Widgets/Types/SMessagingTypes.h
--- a/Engine/Plugins/Messaging/MessagingDebugger/Source/MessagingDebugger/Private/Widgets/Types/SMessagingTypes.h
+++ b/Engine/Plugins/Messaging/MessagingDebugger/Source/MessagingDebugger/Private/Widgets/Types/SMessagingTypes.h
@@ -16,10 +16,8 @@ public:
 
 public:
 
-	/**
-	 * Destructor.
-	 */
-	~SMessagingTypes( );
+	/** Destructor. */
+	~SMessagingTypes();
 
 public:
 
@@ -42,51 +40,49 @@ protected:
 	 */
 	void AddType( const FMessageTracerTypeInfoRef& TypeInfo );
 
-	/**
-	 * Reloads the message history.
-	 */
-	void ReloadTypes( );
+	/** Reloads the message history. */
+	void ReloadTypes();
 
 private:
 
-	// Callback for message type filter changes.
-	void HandleFilterChanged( );
+	/** Callback for message type filter changes. */
+	void HandleFilterChanged();
 
-	// Callback for handling message selection changes.
-	void HandleModelSelectedMessageChanged( );
+	/** Callback for handling message selection changes. */
+	void HandleModelSelectedMessageChanged();
 
-	// Callback for when the tracer's message history has been reset.
-	void HandleTracerMessagesReset( );
+	/** Callback for when the tracer's message history has been reset. */
+	void HandleTracerMessagesReset();
 
-	// Callback for when a message type has been added to the message tracer.
+	/** Callback for when a message type has been added to the message tracer. */
 	void HandleTracerTypeAdded( FMessageTracerTypeInfoRef TypeInfo );
 
-	// Callback for generating a row widget for the message type list view.
+	/** Callback for generating a row widget for the message type list view. */
 	TSharedRef<ITableRow> HandleTypeListGenerateRow( FMessageTracerTypeInfoPtr TypeInfo, const TSharedRef<STableViewBase>& OwnerTable );
 
-	// Callback for getting the highlight string for message types.
-	FText HandleTypeListGetHighlightText( ) const;
+	/** Callback for getting the highlight string for message types. */
+	FText HandleTypeListGetHighlightText() const;
 
-	// Callback for selecting message types.
+	/** Callback for selecting message types. */
 	void HandleTypeListSelectionChanged( FMessageTracerTypeInfoPtr InItem, ESelectInfo::Type SelectInfo );
 
 private:
 
-	// Holds the message type filter model.
+	/** Holds the message type filter model. */
 	FMessagingDebuggerTypeFilterPtr Filter;
 
-	// Holds a pointer to the view model.
+	/** Holds a pointer to the view model. */
 	FMessagingDebuggerModelPtr Model;
 
-	// Holds the widget's visual style.
+	/** Holds the widget's visual style. */
 	TSharedPtr<ISlateStyle> Style;
 
-	// Holds a pointer to the message bus tracer.
+	/** Holds a pointer to the message bus tracer. */
 	IMessageTracerPtr Tracer;
 
-	// Holds the filtered list of message types.
+	/** Holds the filtered list of message types. */
 	TArray<FMessageTracerTypeInfoPtr> TypeList;
 
-	// Holds the message type list view.
+	/** Holds the message type list view. */
 	TSharedPtr<SListView<FMessageTracerTypeInfoPtr>> TypeListView;
 };
diff --git a/Engine/Plugins/Messaging/MessagingDebugger/Source/MessagingDebugger/Private/Widgets/Types/SMessagingTypesFilterBar.cpp b/Engine/Plugins/Messaging/MessagingDebugger/Source/MessagingDebugger/Private/Widgets/Types/SMessagingTypesFilterBar.cpp
--- a/Engine/Plugins/Messaging/MessagingDebugger/Source/MessagingDebugger/Private/Widgets/Types/SMessagingTypesFilterBar.cpp
+++ b/Engine/Plugins/Messaging/MessagingDebugger/Source/MessagingDebugger/Private/Widgets/Types/SMessagingTypesFilterBar.cpp
@@ -1,12 +1,16 @@
 // Copyright 1998-2014 Epic Games, Inc. All Rights Reserved.
 
 #include "MessagingDebuggerPrivatePCH.h"
-
 #include "SMessagingTypesFilterBar.h"
 #include "SSearchBox.h"
 
+
 #define LOCTEXT_NAMESPACE "SMessagingTypesFilterBar"
 
+
+/* SMessagingTypesFilterBar interface
+*****************************************************************************/
+
 void SMessagingTypesFilterBar::Construct(const FArguments& InArgs, FMessagingDebuggerTypeFilterRef InFilter)
 {
 	Filter = InFilter;
diff --git a/Engine/Plugins/Messaging/MessagingDebugger/Source/MessagingDebugger/Private/Widgets/Types/SMessagingTypesFilterBar.h b/Engine/Plugins/Messaging/MessagingDebugger/Source/MessagingDebugger/Private/Widgets/Types/SMessagingTypesFilterBar.h
--- a/Engine/Plugins/Messaging/MessagingDebugger/Source/MessagingDebugger/Private/Widgets/Types/SMessagingTypesFilterBar.h
+++ b/Engine/Plugins/Messaging/MessagingDebugger/Source/MessagingDebugger/Private/Widgets/Types/SMessagingTypesFilterBar.h
@@ -18,15 +18,6 @@ public:
 
 	SLATE_END_ARGS()
 
-public:
-
-	/**
-	 * Default constructor.
-	 */
-	SMessagingTypesFilterBar( )
-		: Filter(NULL)
-	{ }
-
 public:
 
 	/**
@@ -39,7 +30,7 @@ public:
 
 private:
 
-	// Handles changing the filter string text box text.
+	/** Handles changing the filter string text box text. */
 	void HandleFilterStringTextChanged( const FText& NewText )
 	{
 		Filter->SetFilterString(NewText.ToString());
@@ -47,6 +38,6 @@ private:
 
 private:
 
-	// Holds the filter model.
+	/** Holds the filter model. */
 	FMessagingDebuggerTypeFilterPtr Filter;
 };
diff --git a/Engine/Plugins/Messaging/MessagingDebugger/Source/MessagingDebugger/Private/Widgets/Types/SMessagingTypesTableRow.h b/Engine/Plugins/Messaging/MessagingDebugger/Source/MessagingDebugger/Private/Widgets/Types/SMessagingTypesTableRow.h
--- a/Engine/Plugins/Messaging/MessagingDebugger/Source/MessagingDebugger/Private/Widgets/Types/SMessagingTypesTableRow.h
+++ b/Engine/Plugins/Messaging/MessagingDebugger/Source/MessagingDebugger/Private/Widgets/Types/SMessagingTypesTableRow.h
@@ -99,20 +99,20 @@ public:
 
 private:
 
-	// Gets the text for the Messages column.
-	FString HandleMessagesText( ) const
+	/** Gets the text for the Messages column. */
+	FString HandleMessagesText() const
 	{
 		return FString::Printf(TEXT("%i"), TypeInfo->Messages.Num());
 	}
 
-	// Handles changing the checked state of the visibility check box.
+	/** Handles changing the checked state of the visibility check box. */
 	void HandleVisibilityCheckBoxCheckStateChanged( ESlateCheckBoxState::Type CheckState )
 	{
 		Model->SetTypeVisibility(TypeInfo.ToSharedRef(), (CheckState == ESlateCheckBoxState::Checked));
 	}
 
-	// Gets the image for the visibility check box.
-	ESlateCheckBoxState::Type HandleVisibilityCheckBoxIsChecked( ) const
+	/** Gets the image for the visibility check box. */
+	ESlateCheckBoxState::Type HandleVisibilityCheckBoxIsChecked() const
 	{
 		if (Model->IsTypeVisible(TypeInfo.ToSharedRef()))
 		{
@@ -124,16 +124,16 @@ private:
 
 private:
 
-	// Holds the highlight string for the message.
+	/** Holds the highlight string for the message. */
 	TAttribute<FText> HighlightText;
 
-	// Holds a pointer to the view model.
+	/** Holds a pointer to the view model. */
 	FMessagingDebuggerModelPtr Model;
 
-	// Holds the widget's visual style.
+	/** Holds the widget's visual style. */
 	TSharedPtr<ISlateStyle> Style;
 
-	// Holds message type's debug information.
+	/** Holds message type's debug information. */
 	FMessageTracerTypeInfoPtr TypeInfo;
 };
 
