MessagingDebugger: Removed redundant info from the Message Details tab

diff --git a/Engine/Plugins/Messaging/MessagingDebugger/Source/MessagingDebugger/Private/Widgets/MessageDetails/SMessagingMessageDetails.cpp b/Engine/Plugins/Messaging/MessagingDebugger/Source/MessagingDebugger/Private/Widgets/MessageDetails/SMessagingMessageDetails.cpp
--- a/Engine/Plugins/Messaging/MessagingDebugger/Source/MessagingDebugger/Private/Widgets/MessageDetails/SMessagingMessageDetails.cpp
+++ b/Engine/Plugins/Messaging/MessagingDebugger/Source/MessagingDebugger/Private/Widgets/MessageDetails/SMessagingMessageDetails.cpp
@@ -37,38 +37,6 @@ void SMessagingMessageDetails::Construct( const FArguments& InArgs, const FMessa
 				SNew(SGridPanel)
 					.FillColumn(1, 1.0f)
 
-				// Message type
-				+ SGridPanel::Slot(0, 0)
-					.Padding(0.0f, 4.0f)
-					[
-						SNew(STextBlock)
-							.Text(LOCTEXT("MessageTypeLabel", "Message Type:"))
-					]
-
-				+ SGridPanel::Slot(1, 0)
-					.HAlign(HAlign_Right)
-					.Padding(0.0f, 4.0f)
-					[
-						SNew(STextBlock)
-							.Text(this, &SMessagingMessageDetails::HandleMessageTypeText)
-					]
-
-				// Sender
-				+ SGridPanel::Slot(0, 1)
-					.Padding(0.0f, 4.0f)
-					[
-						SNew(STextBlock)
-							.Text(LOCTEXT("SenderLabel", "Sender:"))
-					]
-
-				+ SGridPanel::Slot(1, 1)
-					.HAlign(HAlign_Right)
-					.Padding(0.0f, 4.0f)
-					[
-						SNew(STextBlock)
-							.Text(this, &SMessagingMessageDetails::HandleSenderText)
-					]
-
 				// Sender thread
 				+ SGridPanel::Slot(0, 2)
 					.Padding(0.0f, 4.0f)
@@ -85,7 +53,7 @@ void SMessagingMessageDetails::Construct( const FArguments& InArgs, const FMessa
 						.Text(this, &SMessagingMessageDetails::HandleSenderThreadText)
 					]
 
-				// Time sent
+				// Timestamp
 				+ SGridPanel::Slot(0, 3)
 					.Padding(0.0f, 4.0f)
 					[
@@ -215,38 +183,12 @@ FText SMessagingMessageDetails::HandleExpirationText() const
 }
 
 
-FText SMessagingMessageDetails::HandleMessageTypeText() const
-{
-	FMessageTracerMessageInfoPtr SelectedMessage = Model->GetSelectedMessage();
-
-	if (SelectedMessage.IsValid() && SelectedMessage->TypeInfo.IsValid())
-	{
-		return FText::FromName(SelectedMessage->TypeInfo->TypeName);
-	}
-
-	return FText::GetEmpty();
-}
-
-
 void SMessagingMessageDetails::HandleModelSelectedMessageChanged()
 {
 	RefreshDetails();	
 }
 
 
-FText SMessagingMessageDetails::HandleSenderText() const
-{
-	FMessageTracerMessageInfoPtr SelectedMessage = Model->GetSelectedMessage();
-
-	if (SelectedMessage.IsValid() && SelectedMessage->SenderInfo.IsValid())
-	{
-		return FText::FromName(SelectedMessage->SenderInfo->Name);
-	}
-
-	return FText::GetEmpty();
-}
-
-
 FText SMessagingMessageDetails::HandleSenderThreadText() const
 {
 	FMessageTracerMessageInfoPtr SelectedMessage = Model->GetSelectedMessage();
diff --git a/Engine/Plugins/Messaging/MessagingDebugger/Source/MessagingDebugger/Private/Widgets/MessageDetails/SMessagingMessageDetails.h b/Engine/Plugins/Messaging/MessagingDebugger/Source/MessagingDebugger/Private/Widgets/MessageDetails/SMessagingMessageDetails.h
--- a/Engine/Plugins/Messaging/MessagingDebugger/Source/MessagingDebugger/Private/Widgets/MessageDetails/SMessagingMessageDetails.h
+++ b/Engine/Plugins/Messaging/MessagingDebugger/Source/MessagingDebugger/Private/Widgets/MessageDetails/SMessagingMessageDetails.h
@@ -43,15 +43,9 @@ private:
 	/** Callback for getting the text of the 'Expiration' field. */
 	FText HandleExpirationText() const;
 
-	/** Callback for getting the text of the 'Message Type' field. */
-	FText HandleMessageTypeText() const;
-
 	/** Callback for handling message selection changes. */
 	void HandleModelSelectedMessageChanged();
 
-	/** Callback for getting the text of the 'Sender' field. */
-	FText HandleSenderText() const;
-
 	/** Callback for getting the text of the 'Sender Thread' field. */
 	FText HandleSenderThreadText() const;
 
