MessagingDebugger: Moved Message Details tab into center column

diff --git a/Engine/Plugins/Messaging/MessagingDebugger/Source/MessagingDebugger/Private/Widgets/SMessagingDebugger.cpp b/Engine/Plugins/Messaging/MessagingDebugger/Source/MessagingDebugger/Private/Widgets/SMessagingDebugger.cpp
--- a/Engine/Plugins/Messaging/MessagingDebugger/Source/MessagingDebugger/Private/Widgets/SMessagingDebugger.cpp
+++ b/Engine/Plugins/Messaging/MessagingDebugger/Source/MessagingDebugger/Private/Widgets/SMessagingDebugger.cpp
@@ -156,13 +156,14 @@ void SMessagingDebugger::Construct( const FArguments& InArgs, const TSharedRef<S
 							FTabManager::NewStack()
 								->AddTab(MessageHistoryTabId, ETabState::OpenedTab)
 								->SetHideTabWell(true)
-								->SetSizeCoefficient(0.75f)
+								->SetSizeCoefficient(0.725f)
 						)
 						->Split
 						(
 							FTabManager::NewStack()
 								->AddTab(BreakpointsTabId, ETabState::OpenedTab)
-								->SetSizeCoefficient(0.25f)
+								->AddTab(MessageDetailsTabId, ETabState::OpenedTab)
+								->SetSizeCoefficient(0.275f)
 						)
 				)
 				->Split
@@ -180,7 +181,6 @@ void SMessagingDebugger::Construct( const FArguments& InArgs, const TSharedRef<S
 						->Split
 						(
 							FTabManager::NewStack()
-								->AddTab(MessageDetailsTabId, ETabState::OpenedTab)
 								->AddTab(MessageDataTabId, ETabState::OpenedTab)
 								->SetForegroundTab(MessageDetailsTabId)
 								->SetSizeCoefficient(0.4f)
