MessagingDebugger: Layout tweaks

diff --git a/Engine/Plugins/Messaging/MessagingDebugger/Source/MessagingDebugger/Private/Widgets/EndpointDetails/SMessagingEndpointDetails.cpp b/Engine/Plugins/Messaging/MessagingDebugger/Source/MessagingDebugger/Private/Widgets/EndpointDetails/SMessagingEndpointDetails.cpp
--- a/Engine/Plugins/Messaging/MessagingDebugger/Source/MessagingDebugger/Private/Widgets/EndpointDetails/SMessagingEndpointDetails.cpp
+++ b/Engine/Plugins/Messaging/MessagingDebugger/Source/MessagingDebugger/Private/Widgets/EndpointDetails/SMessagingEndpointDetails.cpp
@@ -60,7 +60,7 @@ void SMessagingEndpointDetails::Construct( const FArguments& InArgs, const FMess
 
 		+ SVerticalBox::Slot()
 			.FillHeight(1.0f)
-			.Padding(0.0f, 4.0f, 0.0f, 0.0f)
+			.Padding(0.0f, 8.0f, 0.0f, 0.0f)
 			[
 				SNew(SBorder)
 					.BorderImage(InStyle->GetBrush("GroupBorder"))
@@ -77,7 +77,7 @@ void SMessagingEndpointDetails::Construct( const FArguments& InArgs, const FMess
 								SNew(SHeaderRow)
 
 								+ SHeaderRow::Column("Address")
-									.DefaultLabel(FText::FromString(TEXT("Endpoint Address")))
+									.DefaultLabel(FText::FromString(TEXT("Addresses")))
 									.FillWidth(1.0f)
 
 								+ SHeaderRow::Column("TimeRegistered")
diff --git a/Engine/Plugins/Messaging/MessagingDebugger/Source/MessagingDebugger/Private/Widgets/MessageDetails/SMessagingDispatchStateTableRow.h b/Engine/Plugins/Messaging/MessagingDebugger/Source/MessagingDebugger/Private/Widgets/MessageDetails/SMessagingDispatchStateTableRow.h
--- a/Engine/Plugins/Messaging/MessagingDebugger/Source/MessagingDebugger/Private/Widgets/MessageDetails/SMessagingDispatchStateTableRow.h
+++ b/Engine/Plugins/Messaging/MessagingDebugger/Source/MessagingDebugger/Private/Widgets/MessageDetails/SMessagingDispatchStateTableRow.h
@@ -53,39 +53,53 @@ public:
 				.VAlign(VAlign_Center)
 				[
 					SNew(STextBlock)
-					.Text(TimespanToReadableText(DispatchState->DispatchLatency))
+						.Text(TimespanToReadableText(DispatchState->DispatchLatency))
 				];
 		}
-		else if (ColumnName == "HandleTime")
+		else if (ColumnName == "DispatchType")
 		{
 			return SNew(SBox)
 				.Padding(FMargin(4.0f, 0.0f))
 				.VAlign(VAlign_Center)
 				[
 					SNew(STextBlock)
-						.ColorAndOpacity(this, &SMessagingDispatchStateTableRow::HandleHandlingTimeColorAndOpacity)
-						.Text(this, &SMessagingDispatchStateTableRow::HandleHandlingTimeText)
+						.Text(this, &SMessagingDispatchStateTableRow::HandleDispatchTypeText)
+						.ToolTipText(this, &SMessagingDispatchStateTableRow::HandleDispatchTypeTooltip)
 				];
 		}
-		else if (ColumnName == "Recipient")
+		else if (ColumnName == "HandleTime")
 		{
 			return SNew(SBox)
 				.Padding(FMargin(4.0f, 0.0f))
 				.VAlign(VAlign_Center)
 				[
 					SNew(STextBlock)
-						.Text(FText::FromName(DispatchState->EndpointInfo->Name))
+						.ColorAndOpacity(this, &SMessagingDispatchStateTableRow::HandleHandlingTimeColorAndOpacity)
+						.Text(this, &SMessagingDispatchStateTableRow::HandleHandlingTimeText)
 				];
 		}
-		else if (ColumnName == "Type")
+		else if (ColumnName == "Recipient")
 		{
-			return SNew(SBox)
-				.HAlign(HAlign_Center)
+			FMessageTracerEndpointInfoPtr EndpointInfo = DispatchState->EndpointInfo;
+
+			return SNew(SHorizontalBox)
+
+			+ SHorizontalBox::Slot()
+				.AutoWidth()
+				.HAlign(HAlign_Left)
 				.VAlign(VAlign_Center)
+				.Padding(FMargin(4.0f, 0.0f))
 				[
 					SNew(SImage)
-						.Image(this, &SMessagingDispatchStateTableRow::HandleTypeImage)
-						.ToolTipText(this, &SMessagingDispatchStateTableRow::HandleTypeTooltip)
+						.Image(Style->GetBrush(EndpointInfo->Remote ? "RemoteEndpoint" : "LocalEndpoint"))
+						.ToolTipText(EndpointInfo->Remote ? LOCTEXT("RemoteEndpointTooltip", "Remote Endpoint") : LOCTEXT("LocalEndpointTooltip", "Local Endpoint"))
+				]
+
+			+ SHorizontalBox::Slot()
+				.VAlign(VAlign_Center)
+				[
+					SNew(STextBlock)
+						.Text(EndpointInfo->Name.ToString())
 				];
 		}
 
@@ -159,6 +173,41 @@ protected:
 
 private:
 
+	/** Callback for getting the dispatch type. */
+	FText HandleDispatchTypeText() const
+	{
+		switch (DispatchState->DispatchType)
+		{
+		case EMessageTracerDispatchTypes::Direct:
+			return LOCTEXT("DispatchTypeDirect", "Direct");
+
+		case EMessageTracerDispatchTypes::Pending:
+			return LOCTEXT("DispatchTypeDirect", "Pending");
+
+		case EMessageTracerDispatchTypes::TaskGraph:
+			return LOCTEXT("DispatchTypeDirect", "TaskGraph");
+
+		default:
+			return FText::GetEmpty();
+		}
+	}
+
+	/** Callback for getting the dispatch type tool tip text. */
+	FText HandleDispatchTypeTooltip() const
+	{
+		if (DispatchState->DispatchType == EMessageTracerDispatchTypes::Direct)
+		{
+			return LOCTEXT("DispatchDirectTooltip", "Dispatched directly (synchronously)");
+		}
+
+		if (DispatchState->DispatchType == EMessageTracerDispatchTypes::TaskGraph)
+		{
+			return LOCTEXT("DispatchTaskGraphTooltip", "Dispatched with Task Graph (asynchronously)");
+		}
+
+		return LOCTEXT("DispatchPendingTooltip", "Dispatched pending");
+	}
+
 	/** Callback for getting the handling time text. */
 	FText HandleHandlingTimeText() const
 	{
@@ -197,22 +246,6 @@ private:
 		return Style->GetBrush("DispatchPending");
 	}
 
-	/** Callback for getting the dispatch type tool tip text. */
-	FText HandleTypeTooltip() const
-	{
-		if (DispatchState->DispatchType == EMessageTracerDispatchTypes::Direct)
-		{
-			return LOCTEXT("DispatchDirectTooltip", "Dispatched directly (synchronously)");
-		}
-
-		if (DispatchState->DispatchType == EMessageTracerDispatchTypes::TaskGraph)
-		{
-			return LOCTEXT("DispatchTaskGraphTooltip", "Dispatched with Task Graph (asynchronously)");
-		}
-
-		return LOCTEXT("DispatchPendingTooltip", "Dispatched pending");
-	}
-
 private:
 
 	/** Holds the message dispatch state. */
diff --git a/Engine/Plugins/Messaging/MessagingDebugger/Source/MessagingDebugger/Private/Widgets/MessageDetails/SMessagingMessageDetails.cpp b/Engine/Plugins/Messaging/MessagingDebugger/Source/MessagingDebugger/Private/Widgets/MessageDetails/SMessagingMessageDetails.cpp
--- a/Engine/Plugins/Messaging/MessagingDebugger/Source/MessagingDebugger/Private/Widgets/MessageDetails/SMessagingMessageDetails.cpp
+++ b/Engine/Plugins/Messaging/MessagingDebugger/Source/MessagingDebugger/Private/Widgets/MessageDetails/SMessagingMessageDetails.cpp
@@ -37,6 +37,7 @@ void SMessagingMessageDetails::Construct( const FArguments& InArgs, const FMessa
 				SNew(SGridPanel)
 					.FillColumn(1, 1.0f)
 
+
 				// Sender thread
 				+ SGridPanel::Slot(0, 2)
 					.Padding(0.0f, 4.0f)
@@ -88,7 +89,7 @@ void SMessagingMessageDetails::Construct( const FArguments& InArgs, const FMessa
 
 		+ SVerticalBox::Slot()
 			.FillHeight(1.0f)
-			.Padding(0.0f, 4.0f, 0.0f, 0.0f)
+			.Padding(0.0f, 8.0f, 0.0f, 0.0f)
 			[
 				SNew(SBorder)
 					.BorderImage(InStyle->GetBrush("GroupBorder"))
@@ -104,13 +105,13 @@ void SMessagingMessageDetails::Construct( const FArguments& InArgs, const FMessa
 							(
 								SNew(SHeaderRow)
 
-								+ SHeaderRow::Column("Type")
-									.DefaultLabel( FText::FromString(TEXT(" ")))
-									.FixedWidth(20.0f)
-
 								+ SHeaderRow::Column("Recipient")
-									.DefaultLabel(LOCTEXT("DispatchStateListRecipientColumnHeader", "Recipient"))
-									.FillWidth(1.0f)
+									.DefaultLabel(LOCTEXT("DispatchStateListRecipientColumnHeader", "Recipient Endpoint"))
+									.FillWidth(0.75f)
+
+								+ SHeaderRow::Column("DispatchType")
+									.DefaultLabel(LOCTEXT("DispatchStateListDispatchTypeColumnHeader", "Dispatch Type"))
+									.FillWidth(0.25f)
 
 								+ SHeaderRow::Column("DispatchLatency")
 									.DefaultLabel(LOCTEXT("DispatchStateListDispatchedColumnHeader", "Dispatch Latency"))
diff --git a/Engine/Plugins/Messaging/MessagingDebugger/Source/MessagingDebugger/Private/Widgets/SMessagingDebugger.cpp b/Engine/Plugins/Messaging/MessagingDebugger/Source/MessagingDebugger/Private/Widgets/SMessagingDebugger.cpp
--- a/Engine/Plugins/Messaging/MessagingDebugger/Source/MessagingDebugger/Private/Widgets/SMessagingDebugger.cpp
+++ b/Engine/Plugins/Messaging/MessagingDebugger/Source/MessagingDebugger/Private/Widgets/SMessagingDebugger.cpp
@@ -122,7 +122,7 @@ void SMessagingDebugger::Construct( const FArguments& InArgs, const TSharedRef<S
 						(
 							FTabManager::NewStack()
 								->AddTab(EndpointsTabId, ETabState::OpenedTab)
-								->SetSizeCoefficient(0.6f)
+								->SetSizeCoefficient(0.5f)
 						)
 						->Split
 						(
@@ -130,7 +130,7 @@ void SMessagingDebugger::Construct( const FArguments& InArgs, const TSharedRef<S
 								->AddTab(EndpointDetailsTabId, ETabState::OpenedTab)
 								->AddTab(InterceptorsTabId, ETabState::OpenedTab)
 								->SetForegroundTab(EndpointDetailsTabId)
-								->SetSizeCoefficient(0.4f)
+								->SetSizeCoefficient(0.5f)
 						)
 				)
 				->Split
@@ -176,14 +176,14 @@ void SMessagingDebugger::Construct( const FArguments& InArgs, const TSharedRef<S
 						(
 							FTabManager::NewStack()
 								->AddTab(MessageTypesTabId, ETabState::OpenedTab)
-								->SetSizeCoefficient(0.6f)
+								->SetSizeCoefficient(0.33f)
 						)
 						->Split
 						(
 							FTabManager::NewStack()
 								->AddTab(MessageDataTabId, ETabState::OpenedTab)
 								->SetForegroundTab(MessageDetailsTabId)
-								->SetSizeCoefficient(0.4f)
+								->SetSizeCoefficient(0.66f)
 						)
 				)
 		);
