CIS fix

diff --git a/Engine/Plugins/Messaging/MessagingDebugger/Source/MessagingDebugger/Private/Widgets/Endpoints/SMessagingEndpoints.h b/Engine/Plugins/Messaging/MessagingDebugger/Source/MessagingDebugger/Private/Widgets/Endpoints/SMessagingEndpoints.h
--- a/Engine/Plugins/Messaging/MessagingDebugger/Source/MessagingDebugger/Private/Widgets/Endpoints/SMessagingEndpoints.h
+++ b/Engine/Plugins/Messaging/MessagingDebugger/Source/MessagingDebugger/Private/Widgets/Endpoints/SMessagingEndpoints.h
@@ -62,7 +62,7 @@ private:
 private:
 
 	/** Holds the filter bar. */
-	/**TSharedPtr<SMessagingEndpointsFilterBar> FilterBar;
+	//TSharedPtr<SMessagingEndpointsFilterBar> FilterBar;
 
 	/** Holds the filtered list of historic messages. */
 	TArray<FMessageTracerEndpointInfoPtr> EndpointList;
