Messaging: Tracing recipient thread in dispatch info

diff --git a/Engine/Plugins/Messaging/MessagingDebugger/Source/MessagingDebugger/Private/Widgets/MessageDetails/SMessagingDispatchStateTableRow.h b/Engine/Plugins/Messaging/MessagingDebugger/Source/MessagingDebugger/Private/Widgets/MessageDetails/SMessagingDispatchStateTableRow.h
--- a/Engine/Plugins/Messaging/MessagingDebugger/Source/MessagingDebugger/Private/Widgets/MessageDetails/SMessagingDispatchStateTableRow.h
+++ b/Engine/Plugins/Messaging/MessagingDebugger/Source/MessagingDebugger/Private/Widgets/MessageDetails/SMessagingDispatchStateTableRow.h
@@ -102,6 +102,16 @@ public:
 						.Text(FText::FromName(EndpointInfo->Name))
 				];
 		}
+		else if (ColumnName == "RecipientThread")
+		{
+			return SNew(SBox)
+				.Padding(FMargin(4.0f, 0.0f))
+				.VAlign(VAlign_Center)
+				[
+					SNew(STextBlock)
+						.Text(NamedThreadToReadableText(DispatchState->RecipientThread))
+				];
+		}
 
 		return SNullWidget::NullWidget;
 	}
@@ -135,6 +145,43 @@ protected:
 		return FSlateColor::UseForeground();
 	}
 
+	/**
+	 * Converts a named thread to a human readable string.
+	 *
+	 * @param NamedThread The named thread.
+	 * @return The text representation.
+	 */
+	FText NamedThreadToReadableText( ENamedThreads::Type NamedThread )
+	{
+		switch (NamedThread)
+		{
+		case ENamedThreads::AnyThread:
+			return LOCTEXT("AnyThread", "AnyThread");
+			break;
+
+		case ENamedThreads::RHIThread:
+			return LOCTEXT("RHIThread", "RHIThread");
+			break;
+
+		case ENamedThreads::GameThread:
+			return LOCTEXT("GameThread", "GameThread");
+			break;
+
+		case ENamedThreads::ActualRenderingThread:
+			return LOCTEXT("ActualRenderingThread", "ActualRenderingThread");
+			break;
+
+#if STATS
+		case ENamedThreads::StatsThread:
+			return LOCTEXT("StatsThread", "StatsThread");
+			break;
+#endif
+
+		default:
+			return LOCTEXT("UnknownThread", "Unknown");
+		}
+	}
+
 	/**
 	 * Converts the given time span in seconds to a human readable string.
 	 *
@@ -205,7 +252,7 @@ private:
 			return LOCTEXT("DispatchTaskGraphTooltip", "Dispatched with Task Graph (asynchronously)");
 		}
 
-		return LOCTEXT("DispatchPendingTooltip", "Dispatched pending");
+		return LOCTEXT("DispatchPendingTooltip", "Dispatch pending");
 	}
 
 	/** Callback for getting the handling time text. */
diff --git a/Engine/Plugins/Messaging/MessagingDebugger/Source/MessagingDebugger/Private/Widgets/MessageDetails/SMessagingMessageDetails.cpp b/Engine/Plugins/Messaging/MessagingDebugger/Source/MessagingDebugger/Private/Widgets/MessageDetails/SMessagingMessageDetails.cpp
--- a/Engine/Plugins/Messaging/MessagingDebugger/Source/MessagingDebugger/Private/Widgets/MessageDetails/SMessagingMessageDetails.cpp
+++ b/Engine/Plugins/Messaging/MessagingDebugger/Source/MessagingDebugger/Private/Widgets/MessageDetails/SMessagingMessageDetails.cpp
@@ -107,12 +107,16 @@ void SMessagingMessageDetails::Construct( const FArguments& InArgs, const FMessa
 
 								+ SHeaderRow::Column("Recipient")
 									.DefaultLabel(LOCTEXT("DispatchStateListRecipientColumnHeader", "Recipient Endpoint"))
-									.FillWidth(0.75f)
+									.FillWidth(0.5f)
 
 								+ SHeaderRow::Column("DispatchType")
 									.DefaultLabel(LOCTEXT("DispatchStateListDispatchTypeColumnHeader", "Dispatch Type"))
 									.FillWidth(0.25f)
 
+								+ SHeaderRow::Column("RecipientThread")
+									.DefaultLabel(LOCTEXT("DispatchStateListRecipientThreadColumnHeader", "Recipient Thread"))
+									.FillWidth(0.25f)
+
 								+ SHeaderRow::Column("DispatchLatency")
 									.DefaultLabel(LOCTEXT("DispatchStateListDispatchedColumnHeader", "Dispatch Latency"))
 									.FixedWidth(112.0f)
@@ -204,6 +208,10 @@ FText SMessagingMessageDetails::HandleSenderThreadText() const
 			return LOCTEXT("AnyThread", "AnyThread");
 			break;
 
+		case ENamedThreads::RHIThread:
+			return LOCTEXT("RHIThread", "RHIThread");
+			break;
+
 		case ENamedThreads::GameThread:
 			return LOCTEXT("GameThread", "GameThread");
 			break;
@@ -212,6 +220,24 @@ FText SMessagingMessageDetails::HandleSenderThreadText() const
 			return LOCTEXT("ActualRenderingThread", "ActualRenderingThread");
 			break;
 
+		case ENamedThreads::GameThread_Local:
+			return LOCTEXT("GameThread_Local", "GameThread_Local");
+			break;
+
+		case ENamedThreads::ActualRenderingThread_Local:
+			return LOCTEXT("ActualRenderingThread_Local", "ActualRenderingThread_Local");
+			break;
+
+#if STATS
+		case ENamedThreads::StatsThread:
+			return LOCTEXT("StatsThread", "StatsThread");
+			break;
+
+		case ENamedThreads::StatsThread_Local:
+			return LOCTEXT("StatsThread_Local", "StatsThread_Local");
+			break;
+#endif
+
 		default:
 			return LOCTEXT("UnknownThread", "Unknown");
 		}
diff --git a/Engine/Source/Runtime/Messaging/Private/Bus/MessageTracer.cpp b/Engine/Source/Runtime/Messaging/Private/Bus/MessageTracer.cpp
--- a/Engine/Source/Runtime/Messaging/Private/Bus/MessageTracer.cpp
+++ b/Engine/Source/Runtime/Messaging/Private/Bus/MessageTracer.cpp
@@ -117,6 +117,7 @@ void FMessageTracer::TraceDispatchedMessage( const IMessageContextRef& Context,
 			DispatchState->DispatchLatency = Timestamp - MessageInfo->TimeSent;
 			DispatchState->DispatchType = Async ? EMessageTracerDispatchTypes::TaskGraph : EMessageTracerDispatchTypes::Direct;
 			DispatchState->EndpointInfo = EndpointInfo;
+			DispatchState->RecipientThread = Recipient->GetRecipientThread();
 			DispatchState->TimeDispatched = Timestamp;
 			DispatchState->TimeHandled = 0.0;
 
diff --git a/Engine/Source/Runtime/Messaging/Public/IMessageTracer.h b/Engine/Source/Runtime/Messaging/Public/IMessageTracer.h
--- a/Engine/Source/Runtime/Messaging/Public/IMessageTracer.h
+++ b/Engine/Source/Runtime/Messaging/Public/IMessageTracer.h
@@ -67,6 +67,9 @@ struct FMessageTracerDispatchState
 	/** Holds the endpoint to which the message was or is being dispatched. */
 	FMessageTracerEndpointInfoPtr EndpointInfo;
 
+	/** The thread on which the endpoint receives messages. */
+	ENamedThreads::Type RecipientThread;
+
 	/** Holds the time at which the message was dispatched. */
 	double TimeDispatched;
 
