MessagingDebugger: Implemented details view for messages

diff --git a/Engine/Plugins/Messaging/MessagingDebugger/Source/MessagingDebugger/MessagingDebugger.Build.cs b/Engine/Plugins/Messaging/MessagingDebugger/Source/MessagingDebugger/MessagingDebugger.Build.cs
--- a/Engine/Plugins/Messaging/MessagingDebugger/Source/MessagingDebugger/MessagingDebugger.Build.cs
+++ b/Engine/Plugins/Messaging/MessagingDebugger/Source/MessagingDebugger/MessagingDebugger.Build.cs
@@ -12,6 +12,7 @@ public class MessagingDebugger : ModuleRules
 				"Core",
 				"CoreUObject",
                 "InputCore",
+                "PropertyEditor",
 				"Serialization",
 				"Slate",
 				"SlateCore",
diff --git a/Engine/Plugins/Messaging/MessagingDebugger/Source/MessagingDebugger/Private/Widgets/MessageData/SMessagingMessageData.cpp b/Engine/Plugins/Messaging/MessagingDebugger/Source/MessagingDebugger/Private/Widgets/MessageData/SMessagingMessageData.cpp
--- a/Engine/Plugins/Messaging/MessagingDebugger/Source/MessagingDebugger/Private/Widgets/MessageData/SMessagingMessageData.cpp
+++ b/Engine/Plugins/Messaging/MessagingDebugger/Source/MessagingDebugger/Private/Widgets/MessageData/SMessagingMessageData.cpp
@@ -1,8 +1,11 @@
 // Copyright 1998-2015 Epic Games, Inc. All Rights Reserved.
 
 #include "MessagingDebuggerPrivatePCH.h"
+#include "IDetailsView.h"
+#include "IStructureDetailsView.h"
 #include "Json.h"
 #include "JsonStructSerializerBackend.h"
+#include "PropertyEditorModule.h"
 #include "StructSerializer.h"
 
 
@@ -30,12 +33,11 @@ void SMessagingMessageData::Construct( const FArguments& InArgs, const FMessagin
 	Style = InStyle;
 
 	// initialize details view
-/*	FDetailsViewArgs DetailsViewArgs;
+	FDetailsViewArgs DetailsViewArgs;
 	{
 		DetailsViewArgs.bAllowSearch = false;
 		DetailsViewArgs.bHideSelectionTip = true;
 		DetailsViewArgs.bLockable = false;
-		DetailsViewArgs.bObjectsUseNameArea = false;
 		DetailsViewArgs.bSearchInitialKeyFocus = true;
 		DetailsViewArgs.bUpdatesFromSelection = false;
 		DetailsViewArgs.NotifyHook = this;
@@ -43,16 +45,24 @@ void SMessagingMessageData::Construct( const FArguments& InArgs, const FMessagin
 		DetailsViewArgs.bShowModifiedPropertiesOption = false;
 	}
 
-	DetailsView = FModuleManager::GetModuleChecked<FPropertyEditorModule>("PropertyEditor").CreateDetailView(DetailsViewArgs);
-	DetailsView->SetEnabled(TAttribute<bool>::Create(TAttribute<bool>::FGetter::CreateSP(this, &SMessagingMessageData::HandleDetailsViewEnabled)));
-	DetailsView->SetVisibility(TAttribute<EVisibility>::Create(TAttribute<EVisibility>::FGetter::CreateSP(this, &SMessagingMessageData::HandleDetailsViewVisibility)));
-*/
+	FStructureDetailsViewArgs StructureViewArgs;
+	{
+		StructureViewArgs.bShowObjects = false;
+		StructureViewArgs.bShowAssets = true;
+		StructureViewArgs.bShowClasses = true;
+		StructureViewArgs.bShowInterfaces = false;
+	}
+
+	DetailsView = FModuleManager::GetModuleChecked<FPropertyEditorModule>("PropertyEditor")
+		.CreateStructureDetailView(DetailsViewArgs, StructureViewArgs, nullptr, LOCTEXT("MessageData", "Message Data"));
+
+//	DetailsView->SetEnabled(TAttribute<bool>::Create(TAttribute<bool>::FGetter::CreateSP(this, &SMessagingMessageData::HandleDetailsViewEnabled)));
+//	DetailsView->SetVisibility(TAttribute<EVisibility>::Create(TAttribute<EVisibility>::FGetter::CreateSP(this, &SMessagingMessageData::HandleDetailsViewVisibility)));
+
 	ChildSlot
-		[
-			//DetailsView.ToSharedRef()
-			SAssignNew(TextBox, SMultiLineEditableTextBox)
-				.IsReadOnly(true)
-		];
+	[
+		DetailsView->GetWidget().ToSharedRef()
+	];
 
 	Model->OnSelectedMessageChanged().AddRaw(this, &SMessagingMessageData::HandleModelSelectedMessageChanged);
 }
@@ -96,25 +106,16 @@ void SMessagingMessageData::HandleModelSelectedMessageChanged()
 
 		if (MessageTypeInfo != nullptr)
 		{
-			FBufferArchive BufferArchive;
-			FJsonStructSerializerBackend Backend(BufferArchive);
-
-			FStructSerializer::Serialize(SelectedMessage->Context->GetMessage(), *MessageTypeInfo, Backend);
-
-			// add string terminator
-			BufferArchive.Add(0);
-			BufferArchive.Add(0);
-
-			TextBox->SetText(FText::FromString(FString((TCHAR*)BufferArchive.GetData()).Replace(TEXT("\t"), TEXT("    "))));
+			DetailsView->SetStructureData(MakeShareable(new FStructOnScope(MessageTypeInfo, (uint8*)SelectedMessage->Context->GetMessage())));
 		}
 		else
 		{
-			TextBox->SetText(FText::Format(LOCTEXT("UnknownMessageTypeFormat", "Unknown message type '{0}'"), FText::FromString(SelectedMessage->Context->GetMessageType().ToString())));
+			DetailsView->SetStructureData(nullptr);
 		}
 	}
 	else
 	{
-		TextBox->SetText(FText::GetEmpty());
+		DetailsView->SetStructureData(nullptr);
 	}
 }
 
diff --git a/Engine/Plugins/Messaging/MessagingDebugger/Source/MessagingDebugger/Private/Widgets/MessageData/SMessagingMessageData.h b/Engine/Plugins/Messaging/MessagingDebugger/Source/MessagingDebugger/Private/Widgets/MessageData/SMessagingMessageData.h
--- a/Engine/Plugins/Messaging/MessagingDebugger/Source/MessagingDebugger/Private/Widgets/MessageData/SMessagingMessageData.h
+++ b/Engine/Plugins/Messaging/MessagingDebugger/Source/MessagingDebugger/Private/Widgets/MessageData/SMessagingMessageData.h
@@ -3,6 +3,9 @@
 #pragma once
 
 
+class IStructureDetailsView;
+
+
 /**
  * Implements the message data panel.
  */
@@ -51,14 +54,11 @@ private:
 private:
 
 	/** Holds the details view. */
-//	TSharedPtr<IDetailsView> DetailsView;
+	TSharedPtr<IStructureDetailsView> DetailsView;
 
 	/** Holds a pointer to the view model. */
 	FMessagingDebuggerModelPtr Model;
 
 	/** Holds the widget's visual style. */
 	TSharedPtr<ISlateStyle> Style;
-
-	/** Holds the details text box. */
-	TSharedPtr<SMultiLineEditableTextBox> TextBox;
 };
