MessagingDebugger: Adjusted layout

diff --git a/Engine/Plugins/Messaging/MessagingDebugger/Source/MessagingDebugger/Private/Widgets/SMessagingDebugger.cpp b/Engine/Plugins/Messaging/MessagingDebugger/Source/MessagingDebugger/Private/Widgets/SMessagingDebugger.cpp
--- a/Engine/Plugins/Messaging/MessagingDebugger/Source/MessagingDebugger/Private/Widgets/SMessagingDebugger.cpp
+++ b/Engine/Plugins/Messaging/MessagingDebugger/Source/MessagingDebugger/Private/Widgets/SMessagingDebugger.cpp
@@ -176,14 +176,14 @@ void SMessagingDebugger::Construct( const FArguments& InArgs, const TSharedRef<S
 						(
 							FTabManager::NewStack()
 								->AddTab(MessageTypesTabId, ETabState::OpenedTab)
-								->SetSizeCoefficient(0.33f)
+								->SetSizeCoefficient(0.5f)
 						)
 						->Split
 						(
 							FTabManager::NewStack()
 								->AddTab(MessageDataTabId, ETabState::OpenedTab)
 								->SetForegroundTab(MessageDetailsTabId)
-								->SetSizeCoefficient(0.66f)
+								->SetSizeCoefficient(0.5f)
 						)
 				)
 		);
