MessagingDebugger: Allowing editing only for messages that haven't been routed yet.

diff --git a/Engine/Plugins/Messaging/MessagingDebugger/Source/MessagingDebugger/Private/Widgets/MessageData/SMessagingMessageData.cpp b/Engine/Plugins/Messaging/MessagingDebugger/Source/MessagingDebugger/Private/Widgets/MessageData/SMessagingMessageData.cpp
--- a/Engine/Plugins/Messaging/MessagingDebugger/Source/MessagingDebugger/Private/Widgets/MessageData/SMessagingMessageData.cpp
+++ b/Engine/Plugins/Messaging/MessagingDebugger/Source/MessagingDebugger/Private/Widgets/MessageData/SMessagingMessageData.cpp
@@ -53,15 +53,15 @@ void SMessagingMessageData::Construct( const FArguments& InArgs, const FMessagin
 		StructureViewArgs.bShowInterfaces = false;
 	}
 
-	DetailsView = FModuleManager::GetModuleChecked<FPropertyEditorModule>("PropertyEditor")
+	StructureDetailsView = FModuleManager::GetModuleChecked<FPropertyEditorModule>("PropertyEditor")
 		.CreateStructureDetailView(DetailsViewArgs, StructureViewArgs, nullptr, LOCTEXT("MessageData", "Message Data"));
 
-//	DetailsView->SetEnabled(TAttribute<bool>::Create(TAttribute<bool>::FGetter::CreateSP(this, &SMessagingMessageData::HandleDetailsViewEnabled)));
-//	DetailsView->SetVisibility(TAttribute<EVisibility>::Create(TAttribute<EVisibility>::FGetter::CreateSP(this, &SMessagingMessageData::HandleDetailsViewVisibility)));
+	StructureDetailsView->GetDetailsView().SetIsPropertyEditingEnabledDelegate(FIsPropertyEditingEnabled::CreateSP(this, &SMessagingMessageData::HandleDetailsViewIsPropertyEditable));
+	StructureDetailsView->GetDetailsView().SetVisibility(TAttribute<EVisibility>::Create(TAttribute<EVisibility>::FGetter::CreateSP(this, &SMessagingMessageData::HandleDetailsViewVisibility)));
 
 	ChildSlot
 	[
-		DetailsView->GetWidget().ToSharedRef()
+		StructureDetailsView->GetWidget().ToSharedRef()
 	];
 
 	Model->OnSelectedMessageChanged().AddRaw(this, &SMessagingMessageData::HandleModelSelectedMessageChanged);
@@ -79,9 +79,16 @@ void SMessagingMessageData::NotifyPostChange( const FPropertyChangedEvent& Prope
 /* SMessagingMessageData callbacks
  *****************************************************************************/
 
-bool SMessagingMessageData::HandleDetailsViewEnabled() const
+bool SMessagingMessageData::HandleDetailsViewIsPropertyEditable() const
 {
-	return true;
+	FMessageTracerMessageInfoPtr SelectedMessage = Model->GetSelectedMessage();
+
+	if (!SelectedMessage.IsValid() || !SelectedMessage->Context.IsValid())
+	{
+		return false;
+	}
+
+	return (SelectedMessage->TimeRouted == 0.0);
 }
 
 
@@ -106,16 +113,16 @@ void SMessagingMessageData::HandleModelSelectedMessageChanged()
 
 		if (MessageTypeInfo != nullptr)
 		{
-			DetailsView->SetStructureData(MakeShareable(new FStructOnScope(MessageTypeInfo, (uint8*)SelectedMessage->Context->GetMessage())));
+			StructureDetailsView->SetStructureData(MakeShareable(new FStructOnScope(MessageTypeInfo, (uint8*)SelectedMessage->Context->GetMessage())));
 		}
 		else
 		{
-			DetailsView->SetStructureData(nullptr);
+			StructureDetailsView->SetStructureData(nullptr);
 		}
 	}
 	else
 	{
-		DetailsView->SetStructureData(nullptr);
+		StructureDetailsView->SetStructureData(nullptr);
 	}
 }
 
diff --git a/Engine/Plugins/Messaging/MessagingDebugger/Source/MessagingDebugger/Private/Widgets/MessageData/SMessagingMessageData.h b/Engine/Plugins/Messaging/MessagingDebugger/Source/MessagingDebugger/Private/Widgets/MessageData/SMessagingMessageData.h
--- a/Engine/Plugins/Messaging/MessagingDebugger/Source/MessagingDebugger/Private/Widgets/MessageData/SMessagingMessageData.h
+++ b/Engine/Plugins/Messaging/MessagingDebugger/Source/MessagingDebugger/Private/Widgets/MessageData/SMessagingMessageData.h
@@ -42,8 +42,8 @@ public:
 
 private:
 
-	/** Handles checking whether the details view is enabled. */
-	bool HandleDetailsViewEnabled() const;
+	/** Handles checking whether the details view is editable. */
+	bool HandleDetailsViewIsPropertyEditable() const;
 
 	/** Handles determining the visibility of the details view. */
 	EVisibility HandleDetailsViewVisibility() const;
@@ -53,8 +53,8 @@ private:
 
 private:
 
-	/** Holds the details view. */
-	TSharedPtr<IStructureDetailsView> DetailsView;
+	/** Holds the structure details view. */
+	TSharedPtr<IStructureDetailsView> StructureDetailsView;
 
 	/** Holds a pointer to the view model. */
 	FMessagingDebuggerModelPtr Model;
