MessagingDebugger: Cleanup pass.

diff --git a/Engine/Plugins/Messaging/MessagingDebugger/Source/MessagingDebugger/Private/Widgets/SMessagingDebugger.cpp b/Engine/Plugins/Messaging/MessagingDebugger/Source/MessagingDebugger/Private/Widgets/SMessagingDebugger.cpp
--- a/Engine/Plugins/Messaging/MessagingDebugger/Source/MessagingDebugger/Private/Widgets/SMessagingDebugger.cpp
+++ b/Engine/Plugins/Messaging/MessagingDebugger/Source/MessagingDebugger/Private/Widgets/SMessagingDebugger.cpp
@@ -35,76 +35,79 @@ SMessagingDebugger::SMessagingDebugger()
 /* SMessagingDebugger interface
  *****************************************************************************/
 
-void SMessagingDebugger::Construct( const FArguments& InArgs, const TSharedRef<SDockTab>& ConstructUnderMajorTab, const TSharedPtr<SWindow>& ConstructUnderWindow, const IMessageTracerRef& InMessageTracer, const TSharedRef<ISlateStyle>& InStyle )
+void SMessagingDebugger::Construct(const FArguments& InArgs, const TSharedRef<SDockTab>& ConstructUnderMajorTab, const TSharedPtr<SWindow>& ConstructUnderWindow, const IMessageTracerRef& InMessageTracer, const TSharedRef<ISlateStyle>& InStyle)
 {
 	MessageTracer = InMessageTracer;
 	Style = InStyle;
 
 	// bind commands
-	const FMessagingDebuggerCommands& Commands = FMessagingDebuggerCommands::Get();
 	FUICommandList& ActionList = *CommandList;
-
-	ActionList.MapAction(Commands.BreakDebugger, FExecuteAction::CreateRaw(this, &SMessagingDebugger::HandleBreakDebuggerCommandExecute), FCanExecuteAction::CreateRaw(this, &SMessagingDebugger::HandleBreakDebuggerCommandCanExecute));
-	ActionList.MapAction(Commands.ClearHistory, FExecuteAction::CreateRaw(this, &SMessagingDebugger::HandleClearHistoryCommandExecute), FCanExecuteAction::CreateRaw(this, &SMessagingDebugger::HandleClearHistoryCommandCanExecute));
-	ActionList.MapAction(Commands.ContinueDebugger, FExecuteAction::CreateRaw(this, &SMessagingDebugger::HandleContinueDebuggerCommandExecute), FCanExecuteAction::CreateRaw(this, &SMessagingDebugger::HandleContinueDebuggerCommandCanExecute), FIsActionChecked(), FIsActionButtonVisible::CreateRaw(this, &SMessagingDebugger::HandleContinueDebuggerCommandIsVisible));
-	ActionList.MapAction(Commands.StartDebugger, FExecuteAction::CreateRaw(this, &SMessagingDebugger::HandleStartDebuggerCommandExecute), FCanExecuteAction::CreateRaw(this, &SMessagingDebugger::HandleStartDebuggerCommandCanExecute), FIsActionChecked(), FIsActionButtonVisible::CreateRaw(this, &SMessagingDebugger::HandleStartDebuggerCommandIsVisible));
-	ActionList.MapAction(Commands.StepDebugger, FExecuteAction::CreateRaw(this, &SMessagingDebugger::HandleStepDebuggerCommandExecute), FCanExecuteAction::CreateRaw(this, &SMessagingDebugger::HandleStepDebuggerCommandCanExecute));
-	ActionList.MapAction(Commands.StopDebugger, FExecuteAction::CreateRaw(this, &SMessagingDebugger::HandleStopDebuggerCommandExecute), FCanExecuteAction::CreateRaw(this, &SMessagingDebugger::HandleStopDebuggerCommandCanExecute));
-
+	{
+		const FMessagingDebuggerCommands& Commands = FMessagingDebuggerCommands::Get();
+
+		ActionList.MapAction(Commands.BreakDebugger, FExecuteAction::CreateRaw(this, &SMessagingDebugger::HandleBreakDebuggerCommandExecute), FCanExecuteAction::CreateRaw(this, &SMessagingDebugger::HandleBreakDebuggerCommandCanExecute));
+		ActionList.MapAction(Commands.ClearHistory, FExecuteAction::CreateRaw(this, &SMessagingDebugger::HandleClearHistoryCommandExecute), FCanExecuteAction::CreateRaw(this, &SMessagingDebugger::HandleClearHistoryCommandCanExecute));
+		ActionList.MapAction(Commands.ContinueDebugger, FExecuteAction::CreateRaw(this, &SMessagingDebugger::HandleContinueDebuggerCommandExecute), FCanExecuteAction::CreateRaw(this, &SMessagingDebugger::HandleContinueDebuggerCommandCanExecute), FIsActionChecked(), FIsActionButtonVisible::CreateRaw(this, &SMessagingDebugger::HandleContinueDebuggerCommandIsVisible));
+		ActionList.MapAction(Commands.StartDebugger, FExecuteAction::CreateRaw(this, &SMessagingDebugger::HandleStartDebuggerCommandExecute), FCanExecuteAction::CreateRaw(this, &SMessagingDebugger::HandleStartDebuggerCommandCanExecute), FIsActionChecked(), FIsActionButtonVisible::CreateRaw(this, &SMessagingDebugger::HandleStartDebuggerCommandIsVisible));
+		ActionList.MapAction(Commands.StepDebugger, FExecuteAction::CreateRaw(this, &SMessagingDebugger::HandleStepDebuggerCommandExecute), FCanExecuteAction::CreateRaw(this, &SMessagingDebugger::HandleStepDebuggerCommandCanExecute));
+		ActionList.MapAction(Commands.StopDebugger, FExecuteAction::CreateRaw(this, &SMessagingDebugger::HandleStopDebuggerCommandExecute), FCanExecuteAction::CreateRaw(this, &SMessagingDebugger::HandleStopDebuggerCommandCanExecute));
+	}
 
 	// create & initialize tab manager
 	TabManager = FGlobalTabmanager::Get()->NewTabManager(ConstructUnderMajorTab);
-	TSharedRef<FWorkspaceItem> AppMenuGroup = TabManager->AddLocalWorkspaceMenuCategory(LOCTEXT("MessagingDebuggerGroupName", "Messaging Debugger"));
-
-	TabManager->RegisterTabSpawner(BreakpointsTabId, FOnSpawnTab::CreateRaw(this, &SMessagingDebugger::HandleTabManagerSpawnTab, BreakpointsTabId))
-		.SetDisplayName(LOCTEXT("BreakpointsTabTitle", "Breakpoints"))
-		.SetGroup(AppMenuGroup)
-		.SetIcon(FSlateIcon(Style->GetStyleSetName(), "BreakpointsTabIcon"));
-
-	TabManager->RegisterTabSpawner(EndpointDetailsTabId, FOnSpawnTab::CreateRaw(this, &SMessagingDebugger::HandleTabManagerSpawnTab, EndpointDetailsTabId))
-		.SetDisplayName(LOCTEXT("EndpointDetailsTabTitle", "Endpoint Details"))
-		.SetGroup(AppMenuGroup)
-		.SetIcon(FSlateIcon(Style->GetStyleSetName(), "EndpointDetailsTabIcon"));
-
-	TabManager->RegisterTabSpawner(EndpointsTabId, FOnSpawnTab::CreateRaw(this, &SMessagingDebugger::HandleTabManagerSpawnTab, EndpointsTabId))
-		.SetDisplayName(LOCTEXT("EndpointsTabTitle", "Endpoints"))
-		.SetGroup(AppMenuGroup)
-		.SetIcon(FSlateIcon(Style->GetStyleSetName(), "EndpointsTabIcon"));
-
-	TabManager->RegisterTabSpawner(InteractionGraphTabId, FOnSpawnTab::CreateRaw(this, &SMessagingDebugger::HandleTabManagerSpawnTab, InteractionGraphTabId))
-		.SetDisplayName(LOCTEXT("InteractionGraphTabTitle", "Interaction Graph"))
-		.SetGroup(AppMenuGroup)
-		.SetIcon(FSlateIcon(Style->GetStyleSetName(), "InteractionGraphTabIcon"));
-
-	TabManager->RegisterTabSpawner(InterceptorsTabId, FOnSpawnTab::CreateRaw(this, &SMessagingDebugger::HandleTabManagerSpawnTab, InterceptorsTabId))
-		.SetDisplayName(LOCTEXT("InterceptorsTabTitle", "Interceptors"))
-		.SetGroup(AppMenuGroup)
-		.SetIcon(FSlateIcon(Style->GetStyleSetName(), "InterceptorsTabIcon"));
-
-	TabManager->RegisterTabSpawner(MessageDataTabId, FOnSpawnTab::CreateRaw(this, &SMessagingDebugger::HandleTabManagerSpawnTab, MessageDataTabId))
-		.SetDisplayName(LOCTEXT("MessageDataTabTitle", "Message Data"))
-		.SetGroup(AppMenuGroup)
-		.SetIcon(FSlateIcon(Style->GetStyleSetName(), "MessageDataTabIcon"));
-
-	TabManager->RegisterTabSpawner(MessageDetailsTabId, FOnSpawnTab::CreateRaw(this, &SMessagingDebugger::HandleTabManagerSpawnTab, MessageDetailsTabId))
-		.SetDisplayName(LOCTEXT("MessageDetailsTabTitle", "Message Details"))
-		.SetGroup(AppMenuGroup)
-		.SetIcon(FSlateIcon(Style->GetStyleSetName(), "MessageDetailsTabIcon"));
-
-	TabManager->RegisterTabSpawner(MessageHistoryTabId, FOnSpawnTab::CreateRaw(this, &SMessagingDebugger::HandleTabManagerSpawnTab, MessageHistoryTabId))
-		.SetDisplayName(LOCTEXT("MessageHistoryTabTitle", "Message History"))
-		.SetGroup(AppMenuGroup)
-		.SetIcon(FSlateIcon(Style->GetStyleSetName(), "MessageHistoryTabIcon"));
-
-	TabManager->RegisterTabSpawner(MessageTypesTabId, FOnSpawnTab::CreateRaw(this, &SMessagingDebugger::HandleTabManagerSpawnTab, MessageTypesTabId))
-		.SetDisplayName(LOCTEXT("MessageTypesTabTitle", "Message Types"))
-		.SetGroup(AppMenuGroup)
-		.SetIcon(FSlateIcon(Style->GetStyleSetName(), "MessageTypesTabIcon"));
-
-	TabManager->RegisterTabSpawner(ToolbarTabId, FOnSpawnTab::CreateRaw(this, &SMessagingDebugger::HandleTabManagerSpawnTab, ToolbarTabId))
-		.SetDisplayName(LOCTEXT("ToolbarTabTitle", "Toolbar"))
-		.SetGroup(AppMenuGroup)
-		.SetIcon(FSlateIcon(Style->GetStyleSetName(), "ToolbarTabIcon"));
+	{
+		TSharedRef<FWorkspaceItem> AppMenuGroup = TabManager->AddLocalWorkspaceMenuCategory(LOCTEXT("MessagingDebuggerGroupName", "Messaging Debugger"));
+
+		TabManager->RegisterTabSpawner(BreakpointsTabId, FOnSpawnTab::CreateRaw(this, &SMessagingDebugger::HandleTabManagerSpawnTab, BreakpointsTabId))
+			.SetDisplayName(LOCTEXT("BreakpointsTabTitle", "Breakpoints"))
+			.SetGroup(AppMenuGroup)
+			.SetIcon(FSlateIcon(Style->GetStyleSetName(), "BreakpointsTabIcon"));
+
+		TabManager->RegisterTabSpawner(EndpointDetailsTabId, FOnSpawnTab::CreateRaw(this, &SMessagingDebugger::HandleTabManagerSpawnTab, EndpointDetailsTabId))
+			.SetDisplayName(LOCTEXT("EndpointDetailsTabTitle", "Endpoint Details"))
+			.SetGroup(AppMenuGroup)
+			.SetIcon(FSlateIcon(Style->GetStyleSetName(), "EndpointDetailsTabIcon"));
+
+		TabManager->RegisterTabSpawner(EndpointsTabId, FOnSpawnTab::CreateRaw(this, &SMessagingDebugger::HandleTabManagerSpawnTab, EndpointsTabId))
+			.SetDisplayName(LOCTEXT("EndpointsTabTitle", "Endpoints"))
+			.SetGroup(AppMenuGroup)
+			.SetIcon(FSlateIcon(Style->GetStyleSetName(), "EndpointsTabIcon"));
+
+		TabManager->RegisterTabSpawner(InteractionGraphTabId, FOnSpawnTab::CreateRaw(this, &SMessagingDebugger::HandleTabManagerSpawnTab, InteractionGraphTabId))
+			.SetDisplayName(LOCTEXT("InteractionGraphTabTitle", "Interaction Graph"))
+			.SetGroup(AppMenuGroup)
+			.SetIcon(FSlateIcon(Style->GetStyleSetName(), "InteractionGraphTabIcon"));
+
+		TabManager->RegisterTabSpawner(InterceptorsTabId, FOnSpawnTab::CreateRaw(this, &SMessagingDebugger::HandleTabManagerSpawnTab, InterceptorsTabId))
+			.SetDisplayName(LOCTEXT("InterceptorsTabTitle", "Interceptors"))
+			.SetGroup(AppMenuGroup)
+			.SetIcon(FSlateIcon(Style->GetStyleSetName(), "InterceptorsTabIcon"));
+
+		TabManager->RegisterTabSpawner(MessageDataTabId, FOnSpawnTab::CreateRaw(this, &SMessagingDebugger::HandleTabManagerSpawnTab, MessageDataTabId))
+			.SetDisplayName(LOCTEXT("MessageDataTabTitle", "Message Data"))
+			.SetGroup(AppMenuGroup)
+			.SetIcon(FSlateIcon(Style->GetStyleSetName(), "MessageDataTabIcon"));
+
+		TabManager->RegisterTabSpawner(MessageDetailsTabId, FOnSpawnTab::CreateRaw(this, &SMessagingDebugger::HandleTabManagerSpawnTab, MessageDetailsTabId))
+			.SetDisplayName(LOCTEXT("MessageDetailsTabTitle", "Message Details"))
+			.SetGroup(AppMenuGroup)
+			.SetIcon(FSlateIcon(Style->GetStyleSetName(), "MessageDetailsTabIcon"));
+
+		TabManager->RegisterTabSpawner(MessageHistoryTabId, FOnSpawnTab::CreateRaw(this, &SMessagingDebugger::HandleTabManagerSpawnTab, MessageHistoryTabId))
+			.SetDisplayName(LOCTEXT("MessageHistoryTabTitle", "Message History"))
+			.SetGroup(AppMenuGroup)
+			.SetIcon(FSlateIcon(Style->GetStyleSetName(), "MessageHistoryTabIcon"));
+
+		TabManager->RegisterTabSpawner(MessageTypesTabId, FOnSpawnTab::CreateRaw(this, &SMessagingDebugger::HandleTabManagerSpawnTab, MessageTypesTabId))
+			.SetDisplayName(LOCTEXT("MessageTypesTabTitle", "Message Types"))
+			.SetGroup(AppMenuGroup)
+			.SetIcon(FSlateIcon(Style->GetStyleSetName(), "MessageTypesTabIcon"));
+
+		TabManager->RegisterTabSpawner(ToolbarTabId, FOnSpawnTab::CreateRaw(this, &SMessagingDebugger::HandleTabManagerSpawnTab, ToolbarTabId))
+			.SetDisplayName(LOCTEXT("ToolbarTabTitle", "Toolbar"))
+			.SetGroup(AppMenuGroup)
+			.SetIcon(FSlateIcon(Style->GetStyleSetName(), "ToolbarTabIcon"));
+	}
 
 	// create tab layout
 	const TSharedRef<FTabManager::FLayout> Layout = FTabManager::NewLayout("MessagingDebuggerLayout_v1.0")
@@ -326,7 +329,7 @@ void SMessagingDebugger::HandleStopDebuggerCommandExecute()
 }
 
 
-TSharedRef<SDockTab> SMessagingDebugger::HandleTabManagerSpawnTab( const FSpawnTabArgs& Args, FName TabIdentifier ) const
+TSharedRef<SDockTab> SMessagingDebugger::HandleTabManagerSpawnTab(const FSpawnTabArgs& Args, FName TabIdentifier) const
 {
 	TSharedPtr<SWidget> TabWidget = SNullWidget::NullWidget;
 	bool AutoSizeTab = false;
