MessagingDebugger: Fixed endpoint address unregistration time never updated in UI

diff --git a/Engine/Plugins/Messaging/MessagingDebugger/Source/MessagingDebugger/Private/Widgets/EndpointDetails/SMessagingAddressTableRow.h b/Engine/Plugins/Messaging/MessagingDebugger/Source/MessagingDebugger/Private/Widgets/EndpointDetails/SMessagingAddressTableRow.h
--- a/Engine/Plugins/Messaging/MessagingDebugger/Source/MessagingDebugger/Private/Widgets/EndpointDetails/SMessagingAddressTableRow.h
+++ b/Engine/Plugins/Messaging/MessagingDebugger/Source/MessagingDebugger/Private/Widgets/EndpointDetails/SMessagingAddressTableRow.h
@@ -6,6 +6,11 @@
 #define LOCTEXT_NAMESPACE "SMessagingAddressTableRow"
 
 
+static const FNumberFormattingOptions TimeRegisteredFormattingOptions = FNumberFormattingOptions()
+	.SetMinimumFractionalDigits(5)
+	.SetMaximumFractionalDigits(5);
+
+
 /**
  * Implements a row widget for the dispatch state list.
  */
@@ -27,7 +32,7 @@ public:
 	 * @param InArgs The construction arguments.
 	 * @param InOwnerTableView The table view that owns this row.
 	 */
-	void Construct( const FArguments& InArgs, const TSharedRef<STableViewBase>& InOwnerTableView, const FMessagingDebuggerModelRef& InModel )
+	void Construct(const FArguments& InArgs, const TSharedRef<STableViewBase>& InOwnerTableView, const FMessagingDebuggerModelRef& InModel)
 	{
 		check(InArgs._Style.IsValid());
 		check(InArgs._AddressInfo.IsValid());
@@ -44,12 +49,8 @@ public:
 	// SMultiColumnTableRow interface
 
 	BEGIN_SLATE_FUNCTION_BUILD_OPTIMIZATION
-	virtual TSharedRef<SWidget> GenerateWidgetForColumn( const FName& ColumnName ) override
+	virtual TSharedRef<SWidget> GenerateWidgetForColumn(const FName& ColumnName) override
 	{
-		static const FNumberFormattingOptions TimeRegisteredFormattingOptions = FNumberFormattingOptions()
-			.SetMinimumFractionalDigits(5)
-			.SetMaximumFractionalDigits(5);
-
 		if (ColumnName == "Address")
 		{
 			return SNew(SBox)
@@ -77,7 +78,7 @@ public:
 				.VAlign(VAlign_Center)
 				[
 					SNew(STextBlock)
-						.Text((AddressInfo->TimeUnregistered > 0.0) ? FText::AsNumber(AddressInfo->TimeUnregistered - GStartTime, &TimeRegisteredFormattingOptions) : LOCTEXT("Never", "Never"))
+						.Text(this, &SMessagingAddressTableRow::HandleTimeUnregisteredText)
 				];
 		}
 
@@ -85,6 +86,19 @@ public:
 	}
 	END_SLATE_FUNCTION_BUILD_OPTIMIZATION
 
+private:
+
+	/** Callback for getting the timestamp at which the address was unregistered. */
+	FText HandleTimeUnregisteredText() const
+	{
+		if (AddressInfo->TimeUnregistered > 0.0)
+		{
+			return FText::AsNumber(AddressInfo->TimeUnregistered - GStartTime, &TimeRegisteredFormattingOptions);
+		}
+
+		return LOCTEXT("Never", "Never");
+	}
+
 private:
 
 	/** Holds the address information. */
