MessagingDebugger: Started to implement interceptors tab.

diff --git a/Engine/Plugins/Messaging/MessagingDebugger/Source/MessagingDebugger/Private/Widgets/Interceptors/SMessagingInterceptors.cpp b/Engine/Plugins/Messaging/MessagingDebugger/Source/MessagingDebugger/Private/Widgets/Interceptors/SMessagingInterceptors.cpp
--- a/Engine/Plugins/Messaging/MessagingDebugger/Source/MessagingDebugger/Private/Widgets/Interceptors/SMessagingInterceptors.cpp
+++ b/Engine/Plugins/Messaging/MessagingDebugger/Source/MessagingDebugger/Private/Widgets/Interceptors/SMessagingInterceptors.cpp
@@ -1,6 +1,7 @@
 // Copyright 1998-2015 Epic Games, Inc. All Rights Reserved.
 
 #include "MessagingDebuggerPrivatePCH.h"
+#include "SMessagingInterceptorsTableRow.h"
 
 
 #define LOCTEXT_NAMESPACE "SMessagingInterceptors"
@@ -9,8 +10,76 @@
 /* SMessagingInterceptors interface
  *****************************************************************************/
 
-void SMessagingInterceptors::Construct( const FArguments& InArgs, const TSharedRef<ISlateStyle>& InStyle, const IMessageTracerRef& InTracer )
+BEGIN_SLATE_FUNCTION_BUILD_OPTIMIZATION
+void SMessagingInterceptors::Construct(const FArguments& InArgs, const FMessagingDebuggerModelRef& InModel, const TSharedRef<ISlateStyle>& InStyle, const IMessageTracerRef& InTracer)
 {
+	Model = InModel;
+	Style = InStyle;
+	Tracer = InTracer;
+
+	ChildSlot
+	[
+		SNew(SVerticalBox)
+
+		+ SVerticalBox::Slot()
+			.FillHeight(1.0f)
+			.Padding(0.0f, 4.0f, 0.0f, 0.0f)
+			[
+				SNew(SBorder)
+					.BorderImage(InStyle->GetBrush("GroupBorder"))
+					.Padding(0.0f)
+					[
+						// interceptor list
+						SAssignNew(InterceptorListView, SListView<FMessageTracerInterceptorInfoPtr>)
+							.ItemHeight(24.0f)
+							.ListItemsSource(&InterceptorList)
+							.SelectionMode(ESelectionMode::None)
+							.OnGenerateRow(this, &SMessagingInterceptors::HandleInterceptorListGenerateRow)
+							.HeaderRow
+							(
+								SNew(SHeaderRow)
+
+								+ SHeaderRow::Column("Name")
+									.DefaultLabel(FText::FromString(TEXT("Interceptors")))
+									.FillWidth(1.0f)
+
+								+ SHeaderRow::Column("TimeRegistered")
+									.DefaultLabel(LOCTEXT("InterceptorListTimeRegisteredColumnHeader", "Time Registered"))
+									.FixedWidth(112.0f)
+									.HAlignCell(HAlign_Right)
+									.HAlignHeader(HAlign_Right)
+
+								+ SHeaderRow::Column("TimeUnregistered")
+									.DefaultLabel(LOCTEXT("InterceptorListTimeUnregisteredColumnHeader", "Time Unregistered"))
+									.FixedWidth(112.0f)
+									.HAlignCell(HAlign_Right)
+									.HAlignHeader(HAlign_Right)
+							)
+					]
+			]
+	];
+}
+END_SLATE_FUNCTION_BUILD_OPTIMIZATION
+
+
+/* SMessagingInterceptors implementation
+ *****************************************************************************/
+
+void SMessagingInterceptors::ReloadInterceptorList()
+{
+	InterceptorList.Empty();
+}
+
+
+/* SMessagingInterceptors callbacks
+ *****************************************************************************/
+
+TSharedRef<ITableRow> SMessagingInterceptors::HandleInterceptorListGenerateRow(FMessageTracerInterceptorInfoPtr InterceptorInfo, const TSharedRef<STableViewBase>& OwnerTable)
+{
+	return SNew(SMessagingInterceptorsTableRow, OwnerTable, Model.ToSharedRef())
+		.InterceptorInfo(InterceptorInfo)
+		.Style(Style);
+//		.ToolTipText(EndpointInfo->Address.ToString(EGuidFormats::DigitsWithHyphensInBraces));
 }
 
 
diff --git a/Engine/Plugins/Messaging/MessagingDebugger/Source/MessagingDebugger/Private/Widgets/Interceptors/SMessagingInterceptors.h b/Engine/Plugins/Messaging/MessagingDebugger/Source/MessagingDebugger/Private/Widgets/Interceptors/SMessagingInterceptors.h
--- a/Engine/Plugins/Messaging/MessagingDebugger/Source/MessagingDebugger/Private/Widgets/Interceptors/SMessagingInterceptors.h
+++ b/Engine/Plugins/Messaging/MessagingDebugger/Source/MessagingDebugger/Private/Widgets/Interceptors/SMessagingInterceptors.h
@@ -20,8 +20,36 @@ public:
 	 * Construct this widget
 	 *
 	 * @param InArgs The declaration data for this widget.
+	 * @param InModel The view model to use.
 	 * @param InStyle The visual style to use for this widget.
 	 * @param InTracer The message tracer.
 	 */
-	void Construct( const FArguments& InArgs, const TSharedRef<ISlateStyle>& InStyle, const IMessageTracerRef& InTracer );
+	void Construct(const FArguments& InArgs, const FMessagingDebuggerModelRef& InModel, const TSharedRef<ISlateStyle>& InStyle, const IMessageTracerRef& InTracer);
+
+protected:
+
+	/** Reloads the list of interceptors. */
+	void ReloadInterceptorList();
+
+private:
+
+	/** Handles generating a row widget for the interceptor list view. */
+	TSharedRef<ITableRow> HandleInterceptorListGenerateRow(FMessageTracerInterceptorInfoPtr InterceptorInfo, const TSharedRef<STableViewBase>& OwnerTable);
+
+private:
+
+	/** Holds the filtered list of interceptors. */
+	TArray<FMessageTracerInterceptorInfoPtr> InterceptorList;
+
+	/** Holds the interceptor list view. */
+	TSharedPtr<SListView<FMessageTracerInterceptorInfoPtr>> InterceptorListView;
+
+	/** Holds a pointer to the view model. */
+	FMessagingDebuggerModelPtr Model;
+
+	/** Holds the widget's visual style. */
+	TSharedPtr<ISlateStyle> Style;
+
+	/** Holds a pointer to the message bus tracer. */
+	IMessageTracerPtr Tracer;
 };
diff --git a/Engine/Plugins/Messaging/MessagingDebugger/Source/MessagingDebugger/Private/Widgets/Interceptors/SMessagingInterceptorsTableRow.h b/Engine/Plugins/Messaging/MessagingDebugger/Source/MessagingDebugger/Private/Widgets/Interceptors/SMessagingInterceptorsTableRow.h
new file mode 100644
--- /dev/null
+++ b/Engine/Plugins/Messaging/MessagingDebugger/Source/MessagingDebugger/Private/Widgets/Interceptors/SMessagingInterceptorsTableRow.h
@@ -0,0 +1,118 @@
+// Copyright 1998-2015 Epic Games, Inc. All Rights Reserved.
+
+#pragma once
+
+
+#define LOCTEXT_NAMESPACE "SMessagingInterceptorTableRow"
+
+
+namespace MessagingInterceptorTableRow
+{
+	static const FNumberFormattingOptions TimeRegisteredFormattingOptions = FNumberFormattingOptions()
+		.SetMinimumFractionalDigits(5)
+		.SetMaximumFractionalDigits(5);
+}
+
+
+/**
+ * Implements a row widget for the interceptors list.
+ */
+class SMessagingInterceptorsTableRow
+	: public SMultiColumnTableRow<FMessageTracerInterceptorInfoPtr>
+{
+public:
+
+	SLATE_BEGIN_ARGS(SMessagingInterceptorsTableRow) { }
+		SLATE_ARGUMENT(FMessageTracerInterceptorInfoPtr, InterceptorInfo)
+		SLATE_ARGUMENT(TSharedPtr<ISlateStyle>, Style)
+	SLATE_END_ARGS()
+
+public:
+
+	/**
+	 * Constructs the widget.
+	 *
+	 * @param InArgs The construction arguments.
+	 * @param InOwnerTableView The table view that owns this row.
+	 */
+	void Construct(const FArguments& InArgs, const TSharedRef<STableViewBase>& InOwnerTableView, const FMessagingDebuggerModelRef& InModel)
+	{
+		check(InArgs._Style.IsValid());
+		check(InArgs._InterceptorInfo.IsValid());
+
+		InterceptorInfo = InArgs._InterceptorInfo;
+		Model = InModel;
+		Style = InArgs._Style;
+
+		SMultiColumnTableRow<FMessageTracerInterceptorInfoPtr>::Construct(FSuperRowType::FArguments(), InOwnerTableView);
+	}
+
+public:
+
+	// SMultiColumnTableRow interface
+
+	BEGIN_SLATE_FUNCTION_BUILD_OPTIMIZATION
+	virtual TSharedRef<SWidget> GenerateWidgetForColumn(const FName& ColumnName) override
+	{
+		if (ColumnName == "Name")
+		{
+			return SNew(SBox)
+				.Padding(FMargin(4.0f, 0.0f))
+				.VAlign(VAlign_Center)
+				[
+					SNew(STextBlock)
+						.Text(FText::FromString(InterceptorInfo->Name.ToString()))
+				];
+		}
+		else if (ColumnName == "TimeRegistered")
+		{
+			return SNew(SBox)
+				.Padding(FMargin(4.0f, 0.0f))
+				.VAlign(VAlign_Center)
+				[
+					SNew(STextBlock)
+						.Text(FText::AsNumber(InterceptorInfo->TimeRegistered - GStartTime, &MessagingInterceptorTableRow::TimeRegisteredFormattingOptions))
+				];
+		}
+		else if (ColumnName == "TimeUnregistered")
+		{
+			return SNew(SBox)
+				.Padding(FMargin(4.0f, 0.0f))
+				.VAlign(VAlign_Center)
+				[
+					SNew(STextBlock)
+						.Text(this, &SMessagingInterceptorsTableRow::HandleTimeUnregisteredText)
+				];
+		}
+
+		return SNullWidget::NullWidget;
+	}
+	END_SLATE_FUNCTION_BUILD_OPTIMIZATION
+
+private:
+
+	/** Callback for getting the timestamp at which the address was unregistered. */
+	FText HandleTimeUnregisteredText() const
+	{
+		if (InterceptorInfo->TimeUnregistered > 0.0)
+		{
+			return FText::AsNumber(InterceptorInfo->TimeUnregistered - GStartTime, &MessagingInterceptorTableRow::TimeRegisteredFormattingOptions);
+		}
+
+		return LOCTEXT("Never", "Never");
+	}
+
+private:
+
+	/** Holds the interceptor information. */
+	FMessageTracerInterceptorInfoPtr InterceptorInfo;
+
+	/** Holds a pointer to the view model. */
+	FMessagingDebuggerModelPtr Model;
+
+	/** Holds the widget's visual style. */
+	TSharedPtr<ISlateStyle> Style;
+};
+
+
+#undef LOCTEXT_NAMESPACE
diff --git a/Engine/Plugins/Messaging/MessagingDebugger/Source/MessagingDebugger/Private/Widgets/SMessagingDebugger.cpp b/Engine/Plugins/Messaging/MessagingDebugger/Source/MessagingDebugger/Private/Widgets/SMessagingDebugger.cpp
--- a/Engine/Plugins/Messaging/MessagingDebugger/Source/MessagingDebugger/Private/Widgets/SMessagingDebugger.cpp
+++ b/Engine/Plugins/Messaging/MessagingDebugger/Source/MessagingDebugger/Private/Widgets/SMessagingDebugger.cpp
@@ -352,7 +352,7 @@ TSharedRef<SDockTab> SMessagingDebugger::HandleTabManagerSpawnTab(const FSpawnTa
 	}
 	else if (TabIdentifier == InterceptorsTabId)
 	{
-		TabWidget = SNew(SMessagingInterceptors, Style.ToSharedRef(), MessageTracer.ToSharedRef());
+		TabWidget = SNew(SMessagingInterceptors, Model, Style.ToSharedRef(), MessageTracer.ToSharedRef());
 	}
 	else if (TabIdentifier == MessageDataTabId)
 	{
