MessagingDebugger: Added keyboard shortcuts for debugger actions

diff --git a/Engine/Plugins/Messaging/MessagingDebugger/Source/MessagingDebugger/Private/Models/MessagingDebuggerCommands.h b/Engine/Plugins/Messaging/MessagingDebugger/Source/MessagingDebugger/Private/Models/MessagingDebuggerCommands.h
--- a/Engine/Plugins/Messaging/MessagingDebugger/Source/MessagingDebugger/Private/Models/MessagingDebuggerCommands.h
+++ b/Engine/Plugins/Messaging/MessagingDebugger/Source/MessagingDebugger/Private/Models/MessagingDebuggerCommands.h
@@ -22,12 +22,12 @@ public:
 
 	virtual void RegisterCommands() override
 	{
-		UI_COMMAND(BreakDebugger, "Break", "Break the debugger at the next message", EUserInterfaceActionType::Button, FInputChord());
+		UI_COMMAND(BreakDebugger, "Break", "Break the debugger at the next message", EUserInterfaceActionType::Button, FInputChord(EKeys::Pause, EModifierKey::Control | EModifierKey::Alt));
 		UI_COMMAND(ClearHistory, "Clear History", "Clears the message history list", EUserInterfaceActionType::Button, FInputChord());
-		UI_COMMAND(ContinueDebugger, "Continue", "Continues debugging", EUserInterfaceActionType::Button, FInputChord());
+		UI_COMMAND(ContinueDebugger, "Continue", "Continues debugging", EUserInterfaceActionType::Button, FInputChord(EKeys::F5));
 		UI_COMMAND(StartDebugger, "Start", "Start the debugger", EUserInterfaceActionType::Button, FInputChord());
-		UI_COMMAND(StepDebugger, "Step", "Step over the current message", EUserInterfaceActionType::Button, FInputChord());
-		UI_COMMAND(StopDebugger, "Stop", "Stop the debugger", EUserInterfaceActionType::Button, FInputChord());
+		UI_COMMAND(StepDebugger, "Step", "Step over the current message", EUserInterfaceActionType::Button, FInputChord(EKeys::F10));
+		UI_COMMAND(StopDebugger, "Stop", "Stop the debugger", EUserInterfaceActionType::Button, FInputChord(EKeys::F5, EModifierKey::Shift));
 	}
 
 public:
