MessagingDebugger: Layout tweaks

diff --git a/Engine/Plugins/Messaging/MessagingDebugger/Source/MessagingDebugger/Private/Widgets/EndpointDetails/SMessagingEndpointDetails.cpp b/Engine/Plugins/Messaging/MessagingDebugger/Source/MessagingDebugger/Private/Widgets/EndpointDetails/SMessagingEndpointDetails.cpp
--- a/Engine/Plugins/Messaging/MessagingDebugger/Source/MessagingDebugger/Private/Widgets/EndpointDetails/SMessagingEndpointDetails.cpp
+++ b/Engine/Plugins/Messaging/MessagingDebugger/Source/MessagingDebugger/Private/Widgets/EndpointDetails/SMessagingEndpointDetails.cpp
@@ -27,7 +27,7 @@ void SMessagingEndpointDetails::Construct( const FArguments& InArgs, const FMess
 
 				// received messages count
 				+ SGridPanel::Slot(0, 0)
-					.Padding(0.0f, 4.0f)
+					.Padding(0.0f, 4.0f, 32.0f, 4.0f)
 					[
 						SNew(STextBlock)
 							.Text(LOCTEXT("EndpointDetailsReceivedMessagesLabel", "Messages Received:"))
@@ -43,7 +43,7 @@ void SMessagingEndpointDetails::Construct( const FArguments& InArgs, const FMess
 
 				// sent messages count
 				+ SGridPanel::Slot(0, 1)
-					.Padding(0.0f, 4.0f)
+					.Padding(0.0f, 4.0f, 32.0f, 4.0f)
 					[
 						SNew(STextBlock)
 							.Text(LOCTEXT("EndpointDetailsReceivedLabel", "Messages Sent:"))
diff --git a/Engine/Plugins/Messaging/MessagingDebugger/Source/MessagingDebugger/Private/Widgets/MessageDetails/SMessagingMessageDetails.cpp b/Engine/Plugins/Messaging/MessagingDebugger/Source/MessagingDebugger/Private/Widgets/MessageDetails/SMessagingMessageDetails.cpp
--- a/Engine/Plugins/Messaging/MessagingDebugger/Source/MessagingDebugger/Private/Widgets/MessageDetails/SMessagingMessageDetails.cpp
+++ b/Engine/Plugins/Messaging/MessagingDebugger/Source/MessagingDebugger/Private/Widgets/MessageDetails/SMessagingMessageDetails.cpp
@@ -40,14 +40,13 @@ void SMessagingMessageDetails::Construct( const FArguments& InArgs, const FMessa
 
 				// Sender thread
 				+ SGridPanel::Slot(0, 2)
-					.Padding(0.0f, 4.0f)
+					.Padding(0.0f, 4.0f, 32.0f, 4.0f)
 					[
 						SNew(STextBlock)
-						.Text(LOCTEXT("SenderThreadLabel", "Sender Thread:"))
+							.Text(LOCTEXT("SenderThreadLabel", "Sender Thread:"))
 					]
 
 				+ SGridPanel::Slot(1, 2)
-					.HAlign(HAlign_Right)
 					.Padding(0.0f, 4.0f)
 					[
 						SNew(STextBlock)
@@ -56,14 +55,13 @@ void SMessagingMessageDetails::Construct( const FArguments& InArgs, const FMessa
 
 				// Timestamp
 				+ SGridPanel::Slot(0, 3)
-					.Padding(0.0f, 4.0f)
+					.Padding(0.0f, 4.0f, 32.0f, 4.0f)
 					[
 						SNew(STextBlock)
 							.Text(LOCTEXT("TimestampLabel", "Timestamp:"))
 					]
 
 				+ SGridPanel::Slot(1, 3)
-					.HAlign(HAlign_Right)
 					.Padding(0.0f, 4.0f)
 					[
 						SNew(STextBlock)
@@ -72,14 +70,13 @@ void SMessagingMessageDetails::Construct( const FArguments& InArgs, const FMessa
 
 				// Expiration
 				+ SGridPanel::Slot(0, 4)
-					.Padding(0.0f, 4.0f)
+					.Padding(0.0f, 4.0f, 32.0f, 4.0f)
 					[
 						SNew(STextBlock)
 							.Text(LOCTEXT("ExpirationLabel", "Expiration:"))
 					]
 
 				+ SGridPanel::Slot(1, 4)
-					.HAlign(HAlign_Right)
 					.Padding(0.0f, 4.0f)
 					[
 						SNew(STextBlock)
