MessagingDebugger: Layout tweaks.

diff --git a/Engine/Plugins/Messaging/MessagingDebugger/Source/MessagingDebugger/Private/Widgets/EndpointDetails/SMessagingEndpointDetails.cpp b/Engine/Plugins/Messaging/MessagingDebugger/Source/MessagingDebugger/Private/Widgets/EndpointDetails/SMessagingEndpointDetails.cpp
--- a/Engine/Plugins/Messaging/MessagingDebugger/Source/MessagingDebugger/Private/Widgets/EndpointDetails/SMessagingEndpointDetails.cpp
+++ b/Engine/Plugins/Messaging/MessagingDebugger/Source/MessagingDebugger/Private/Widgets/EndpointDetails/SMessagingEndpointDetails.cpp
@@ -34,7 +34,6 @@ void SMessagingEndpointDetails::Construct( const FArguments& InArgs, const FMess
 					]
 
 				+ SGridPanel::Slot(1, 0)
-					.HAlign(HAlign_Right)
 					.Padding(0.0f, 4.0f)
 					[
 						SNew(STextBlock)
@@ -50,7 +49,6 @@ void SMessagingEndpointDetails::Construct( const FArguments& InArgs, const FMess
 					]
 
 				+ SGridPanel::Slot(1, 1)
-					.HAlign(HAlign_Right)
 					.Padding(0.0f, 4.0f)
 					[
 						SNew(STextBlock)
