MessagingDebugger: Fallback message data viewer for non-Editor applications (PropertyEditor required UnrealEd at the moment)

diff --git a/Engine/Plugins/Messaging/MessagingDebugger/Source/MessagingDebugger/MessagingDebugger.Build.cs b/Engine/Plugins/Messaging/MessagingDebugger/Source/MessagingDebugger/MessagingDebugger.Build.cs
--- a/Engine/Plugins/Messaging/MessagingDebugger/Source/MessagingDebugger/MessagingDebugger.Build.cs
+++ b/Engine/Plugins/Messaging/MessagingDebugger/Source/MessagingDebugger/MessagingDebugger.Build.cs
@@ -12,7 +12,6 @@ public class MessagingDebugger : ModuleRules
 				"Core",
 				"CoreUObject",
                 "InputCore",
-                "PropertyEditor",
 				"Serialization",
 				"Slate",
 				"SlateCore",
@@ -43,5 +42,10 @@ public class MessagingDebugger : ModuleRules
 				"MessagingDebugger/Private/Widgets/Toolbar",
 			}
 		);
+
+        if (UEBuildConfiguration.bBuildEditor)
+        {
+            PrivateDependencyModuleNames.Add("PropertyEditor");
+        }
 	}
 }
diff --git a/Engine/Plugins/Messaging/MessagingDebugger/Source/MessagingDebugger/Private/Widgets/MessageData/SMessagingMessageData.cpp b/Engine/Plugins/Messaging/MessagingDebugger/Source/MessagingDebugger/Private/Widgets/MessageData/SMessagingMessageData.cpp
--- a/Engine/Plugins/Messaging/MessagingDebugger/Source/MessagingDebugger/Private/Widgets/MessageData/SMessagingMessageData.cpp
+++ b/Engine/Plugins/Messaging/MessagingDebugger/Source/MessagingDebugger/Private/Widgets/MessageData/SMessagingMessageData.cpp
@@ -1,13 +1,16 @@
 // Copyright 1998-2015 Epic Games, Inc. All Rights Reserved.
 
 #include "MessagingDebuggerPrivatePCH.h"
-#include "IDetailsView.h"
-#include "IStructureDetailsView.h"
 #include "Json.h"
 #include "JsonStructSerializerBackend.h"
-#include "PropertyEditorModule.h"
 #include "StructSerializer.h"
 
+#if WITH_EDITOR
+	#include "IDetailsView.h"
+	#include "IStructureDetailsView.h"
+	#include "PropertyEditorModule.h"
+#endif
+
 
 #define LOCTEXT_NAMESPACE "SMessagingMessageData"
 
@@ -32,6 +35,8 @@ void SMessagingMessageData::Construct( const FArguments& InArgs, const FMessagin
 	Model = InModel;
 	Style = InStyle;
 
+#if WITH_EDITOR
+
 	// initialize details view
 	FDetailsViewArgs DetailsViewArgs;
 	{
@@ -64,6 +69,16 @@ void SMessagingMessageData::Construct( const FArguments& InArgs, const FMessagin
 		StructureDetailsView->GetWidget().ToSharedRef()
 	];
 
+#else
+
+	ChildSlot
+	[
+		SAssignNew(TextBox, SMultiLineEditableTextBox)
+			.IsReadOnly(true)
+	];
+
+#endif //WITH_EDITOR
+
 	Model->OnSelectedMessageChanged().AddRaw(this, &SMessagingMessageData::HandleModelSelectedMessageChanged);
 }
 
@@ -113,18 +128,39 @@ void SMessagingMessageData::HandleModelSelectedMessageChanged()
 
 		if (MessageTypeInfo != nullptr)
 		{
+#if WITH_EDITOR
 			StructureDetailsView->SetStructureData(MakeShareable(new FStructOnScope(MessageTypeInfo, (uint8*)SelectedMessage->Context->GetMessage())));
+#else
+			FBufferArchive BufferArchive;
+			FJsonStructSerializerBackend Backend(BufferArchive);
+
+			FStructSerializer::Serialize(SelectedMessage->Context->GetMessage(), *MessageTypeInfo, Backend);
+
+			// add string terminator
+			BufferArchive.Add(0);
+			BufferArchive.Add(0);
+
+			TextBox->SetText(FText::FromString(FString((TCHAR*)BufferArchive.GetData()).Replace(TEXT("\t"), TEXT("    "))));
+#endif //WITH_EDITOR
 		}
 		else
 		{
+#if WITH_EDITOR
 			StructureDetailsView->SetStructureData(nullptr);
+#else
+			TextBox->SetText(FText::Format(LOCTEXT("UnknownMessageTypeFormat", "Unknown message type '{0}'"), FText::FromString(SelectedMessage->Context->GetMessageType().ToString())));
+#endif //WITH_EDITOR
 		}
 	}
 	else
 	{
+#if WITH_EDITOR
 		StructureDetailsView->SetStructureData(nullptr);
+#else
+		TextBox->SetText(FText::GetEmpty());
+#endif //WITH_EDITOR
 	}
 }
 
 
-#undef LOCTEXT_NAMESPACE
\ No newline at end of file
+#undef LOCTEXT_NAMESPACE
diff --git a/Engine/Plugins/Messaging/MessagingDebugger/Source/MessagingDebugger/Private/Widgets/MessageData/SMessagingMessageData.h b/Engine/Plugins/Messaging/MessagingDebugger/Source/MessagingDebugger/Private/Widgets/MessageData/SMessagingMessageData.h
--- a/Engine/Plugins/Messaging/MessagingDebugger/Source/MessagingDebugger/Private/Widgets/MessageData/SMessagingMessageData.h
+++ b/Engine/Plugins/Messaging/MessagingDebugger/Source/MessagingDebugger/Private/Widgets/MessageData/SMessagingMessageData.h
@@ -32,13 +32,13 @@ public:
 	 * @param InModel The view model to use.
 	 * @param InStyle The visual style to use for this widget.
 	 */
-	void Construct( const FArguments& InArgs, const FMessagingDebuggerModelRef& InModel, const TSharedRef<ISlateStyle>& InStyle );
+	void Construct(const FArguments& InArgs, const FMessagingDebuggerModelRef& InModel, const TSharedRef<ISlateStyle>& InStyle);
 
 public:
 
 	// FNotifyHook interface
 
-	virtual void NotifyPostChange( const FPropertyChangedEvent& PropertyChangedEvent, class FEditPropertyChain* PropertyThatChanged ) override;
+	virtual void NotifyPostChange(const FPropertyChangedEvent& PropertyChangedEvent, class FEditPropertyChain* PropertyThatChanged) override;
 
 private:
 
@@ -53,8 +53,13 @@ private:
 
 private:
 
+#if WITH_EDITOR
 	/** Holds the structure details view. */
 	TSharedPtr<IStructureDetailsView> StructureDetailsView;
+#else
+	/** Holds the details text box. */
+	TSharedPtr<SMultiLineEditableTextBox> TextBox;
+#endif //WITH_EDITOR
 
 	/** Holds a pointer to the view model. */
 	FMessagingDebuggerModelPtr Model;
