MessagingDebugger: Enabled keyboard shortcuts

diff --git a/Engine/Plugins/Messaging/MessagingDebugger/Source/MessagingDebugger/Private/Models/MessagingDebuggerCommands.h b/Engine/Plugins/Messaging/MessagingDebugger/Source/MessagingDebugger/Private/Models/MessagingDebuggerCommands.h
--- a/Engine/Plugins/Messaging/MessagingDebugger/Source/MessagingDebugger/Private/Models/MessagingDebuggerCommands.h
+++ b/Engine/Plugins/Messaging/MessagingDebugger/Source/MessagingDebugger/Private/Models/MessagingDebuggerCommands.h
@@ -22,7 +22,7 @@ public:
 
 	virtual void RegisterCommands() override
 	{
-		UI_COMMAND(BreakDebugger, "Break", "Break the debugger at the next message", EUserInterfaceActionType::Button, FInputChord(EKeys::Pause, EModifierKey::Control | EModifierKey::Alt));
+		UI_COMMAND(BreakDebugger, "Break", "Break the debugger at the next message", EUserInterfaceActionType::Button, FInputChord(EKeys::Pause, /*EModifierKey::Control |*/ EModifierKey::Alt));
 		UI_COMMAND(ClearHistory, "Clear History", "Clears the message history list", EUserInterfaceActionType::Button, FInputChord());
 		UI_COMMAND(ContinueDebugger, "Continue", "Continues debugging", EUserInterfaceActionType::Button, FInputChord(EKeys::F5));
 		UI_COMMAND(StartDebugger, "Start", "Start the debugger", EUserInterfaceActionType::Button, FInputChord());
diff --git a/Engine/Plugins/Messaging/MessagingDebugger/Source/MessagingDebugger/Private/Widgets/SMessagingDebugger.cpp b/Engine/Plugins/Messaging/MessagingDebugger/Source/MessagingDebugger/Private/Widgets/SMessagingDebugger.cpp
--- a/Engine/Plugins/Messaging/MessagingDebugger/Source/MessagingDebugger/Private/Widgets/SMessagingDebugger.cpp
+++ b/Engine/Plugins/Messaging/MessagingDebugger/Source/MessagingDebugger/Private/Widgets/SMessagingDebugger.cpp
@@ -236,6 +236,23 @@ void SMessagingDebugger::FillWindowMenu(FMenuBuilder& MenuBuilder, const TShared
 }
 
 
+/* SWidget overrides
+ *****************************************************************************/
+
+FReply SMessagingDebugger::OnKeyDown(const FGeometry& MyGeometry, const FKeyEvent& InKeyEvent)
+{
+	return CommandList->ProcessCommandBindings(InKeyEvent)
+		? FReply::Handled()
+		: FReply::Unhandled();
+}
+
+
+bool SMessagingDebugger::SupportsKeyboardFocus() const
+{
+	return true;
+}
+
+
 /* SMessagingDebugger callbacks
  *****************************************************************************/
 
diff --git a/Engine/Plugins/Messaging/MessagingDebugger/Source/MessagingDebugger/Private/Widgets/SMessagingDebugger.h b/Engine/Plugins/Messaging/MessagingDebugger/Source/MessagingDebugger/Private/Widgets/SMessagingDebugger.h
--- a/Engine/Plugins/Messaging/MessagingDebugger/Source/MessagingDebugger/Private/Widgets/SMessagingDebugger.h
+++ b/Engine/Plugins/Messaging/MessagingDebugger/Source/MessagingDebugger/Private/Widgets/SMessagingDebugger.h
@@ -42,6 +42,13 @@ protected:
 	*/
 	static void FillWindowMenu(FMenuBuilder& MenuBuilder, const TSharedPtr<FTabManager> TabManager);
 
+protected:
+
+	//~ SWidget overrides
+
+	virtual FReply OnKeyDown(const FGeometry& MyGeometry, const FKeyEvent& InKeyEvent) override;
+	virtual bool SupportsKeyboardFocus() const override;
+
 private:
 
 	/** Callback for checking whether the BreakDebugger command can be executed. */
