[INTEGRATE] Engine compilation fixes for VS 2015 Update 1
- Fixed many Slate modules that failing to bracket compiler optimizations correctly (triggers new compiler error)
- Worked around compiler bug triggered by enums in Shared Pointer code
- Disabled compiler warnings about ".." in relative include paths (UE4's build system requires this)

diff --git a/Engine/Plugins/Messaging/MessagingDebugger/Source/MessagingDebugger/Private/Widgets/Endpoints/SMessagingEndpoints.cpp b/Engine/Plugins/Messaging/MessagingDebugger/Source/MessagingDebugger/Private/Widgets/Endpoints/SMessagingEndpoints.cpp
--- a/Engine/Plugins/Messaging/MessagingDebugger/Source/MessagingDebugger/Private/Widgets/Endpoints/SMessagingEndpoints.cpp
+++ b/Engine/Plugins/Messaging/MessagingDebugger/Source/MessagingDebugger/Private/Widgets/Endpoints/SMessagingEndpoints.cpp
@@ -113,7 +113,7 @@ void SMessagingEndpoints::Construct( const FArguments& InArgs, const FMessagingD
 	Filter->OnChanged().AddRaw(this, &SMessagingEndpoints::HandleFilterChanged);
 	Model->OnSelectedMessageChanged().AddRaw(this, &SMessagingEndpoints::HandleModelSelectedMessageChanged);
 }
-
+END_SLATE_FUNCTION_BUILD_OPTIMIZATION
 
 void SMessagingEndpoints::Tick( const FGeometry& AllottedGeometry, const double InCurrentTime, const float InDeltaTime )
 {
diff --git a/Engine/Source/Developer/AITestSuite/Private/AITestSuitePrivatePCH.h b/Engine/Source/Developer/AITestSuite/Private/AITestSuitePrivatePCH.h
--- a/Engine/Source/Developer/AITestSuite/Private/AITestSuitePrivatePCH.h
+++ b/Engine/Source/Developer/AITestSuite/Private/AITestSuitePrivatePCH.h
@@ -40,4 +40,3 @@
 #include "MockAI_BT.h"
 #include "MockGameplayTasks.h"
 
-PRAGMA_DISABLE_OPTIMIZATION
\ No newline at end of file
diff --git a/Engine/Source/Developer/DeviceManager/Private/Widgets/Browser/SDeviceBrowser.cpp b/Engine/Source/Developer/DeviceManager/Private/Widgets/Browser/SDeviceBrowser.cpp
--- a/Engine/Source/Developer/DeviceManager/Private/Widgets/Browser/SDeviceBrowser.cpp
+++ b/Engine/Source/Developer/DeviceManager/Private/Widgets/Browser/SDeviceBrowser.cpp
@@ -105,6 +105,8 @@ void SDeviceBrowser::Construct( const FArguments& InArgs, const FDeviceManagerMo
 
 	Filter->OnFilterChanged().AddSP(this, &SDeviceBrowser::HandleFilterChanged);
 }
+END_SLATE_FUNCTION_BUILD_OPTIMIZATION
+
 void SDeviceBrowser::Tick( const FGeometry& AllottedGeometry, const double InCurrentTime, const float InDeltaTime )
 {
 	//@TODO Passive - Only happens in response to the addition or removal of a device to the device service manager
diff --git a/Engine/Source/Developer/DeviceManager/Private/Widgets/Processes/SDeviceProcesses.cpp b/Engine/Source/Developer/DeviceManager/Private/Widgets/Processes/SDeviceProcesses.cpp
--- a/Engine/Source/Developer/DeviceManager/Private/Widgets/Processes/SDeviceProcesses.cpp
+++ b/Engine/Source/Developer/DeviceManager/Private/Widgets/Processes/SDeviceProcesses.cpp
@@ -142,6 +142,8 @@ EActiveTimerReturnType SDeviceProcesses::UpdateProcessList( double InCurrentTime
 
 	return EActiveTimerReturnType::Continue;
 }
+END_SLATE_FUNCTION_BUILD_OPTIMIZATION
+
 
 /* SDeviceDetails implementation
  *****************************************************************************/
diff --git a/Engine/Source/Developer/DeviceManager/Private/Widgets/Toolbar/SDeviceToolbar.cpp b/Engine/Source/Developer/DeviceManager/Private/Widgets/Toolbar/SDeviceToolbar.cpp
--- a/Engine/Source/Developer/DeviceManager/Private/Widgets/Toolbar/SDeviceToolbar.cpp
+++ b/Engine/Source/Developer/DeviceManager/Private/Widgets/Toolbar/SDeviceToolbar.cpp
@@ -43,7 +43,7 @@ void SDeviceToolbar::Construct( const FArguments& InArgs, const FDeviceManagerMo
 			]
 	];
 }
-
+END_SLATE_FUNCTION_BUILD_OPTIMIZATION
 
 /* SDeviceToolbar callbacks
  *****************************************************************************/
diff --git a/Engine/Source/Developer/Profiler/Private/Widgets/SFiltersAndPresets.cpp b/Engine/Source/Developer/Profiler/Private/Widgets/SFiltersAndPresets.cpp
--- a/Engine/Source/Developer/Profiler/Private/Widgets/SFiltersAndPresets.cpp
+++ b/Engine/Source/Developer/Profiler/Private/Widgets/SFiltersAndPresets.cpp
@@ -758,6 +758,8 @@ void SFiltersAndPresets::Construct( const FArguments& InArgs )
 	CreateGroupByOptionsSources();
 	RecreateSortByOptionsSources();
 }
+END_SLATE_FUNCTION_BUILD_OPTIMIZATION
+
 
 void SFiltersAndPresets::ProfilerManager_OnRequestFilterAndPresetsUpdate()
 {
diff --git a/Engine/Source/Developer/Profiler/Private/Widgets/SProfilerWindow.cpp b/Engine/Source/Developer/Profiler/Private/Widgets/SProfilerWindow.cpp
--- a/Engine/Source/Developer/Profiler/Private/Widgets/SProfilerWindow.cpp
+++ b/Engine/Source/Developer/Profiler/Private/Widgets/SProfilerWindow.cpp
@@ -283,6 +283,7 @@ void SProfilerWindow::Construct( const FArguments& InArgs )
 	FProfilerManager::Get()->OnViewModeChanged().AddSP( this, &SProfilerWindow::ProfilerManager_OnViewModeChanged );
 	GraphPanel->ProfilerMiniView = ProfilerMiniView;
 }
+END_SLATE_FUNCTION_BUILD_OPTIMIZATION
 
 void SProfilerWindow::ManageEventGraphTab( const FGuid ProfilerInstanceID, const bool bCreateFakeTab, const FString TabName )
 {
diff --git a/Engine/Source/Developer/ProjectLauncher/Private/Widgets/SProjectLauncherPreviewPage.cpp b/Engine/Source/Developer/ProjectLauncher/Private/Widgets/SProjectLauncherPreviewPage.cpp
--- a/Engine/Source/Developer/ProjectLauncher/Private/Widgets/SProjectLauncherPreviewPage.cpp
+++ b/Engine/Source/Developer/ProjectLauncher/Private/Widgets/SProjectLauncherPreviewPage.cpp
@@ -786,6 +786,8 @@ void SProjectLauncherPreviewPage::Construct( const FArguments& InArgs, const FPr
 			]
 	];
 }
+END_SLATE_FUNCTION_BUILD_OPTIMIZATION
+
 void SProjectLauncherPreviewPage::Tick( const FGeometry& AllottedGeometry, const double InCurrentTime, const float InDeltaTime )
 {
 	RefreshDeviceProxyList();
diff --git a/Engine/Source/Editor/DetailCustomizations/Private/DialogueStructsCustomizations.cpp b/Engine/Source/Editor/DetailCustomizations/Private/DialogueStructsCustomizations.cpp
--- a/Engine/Source/Editor/DetailCustomizations/Private/DialogueStructsCustomizations.cpp
+++ b/Engine/Source/Editor/DetailCustomizations/Private/DialogueStructsCustomizations.cpp
@@ -782,6 +782,9 @@ void SValidatedDialogueContextHeaderWidget::Construct( const FArguments& InArgs,
 		];
 	}
 }
+END_SLATE_FUNCTION_BUILD_OPTIMIZATION
+
+
 void SValidatedDialogueContextHeaderWidget::Tick( const FGeometry& AllottedGeometry, const double InCurrentTime, const float InDeltaTime )
 {
 	if( DialogueWaveParameterPropertyHandle.IsValid() && DialogueWaveParameterPropertyHandle->IsValidHandle() )
diff --git a/Engine/Source/Editor/DetailCustomizations/Private/DialogueWaveWidgets.cpp b/Engine/Source/Editor/DetailCustomizations/Private/DialogueWaveWidgets.cpp
--- a/Engine/Source/Editor/DetailCustomizations/Private/DialogueWaveWidgets.cpp
+++ b/Engine/Source/Editor/DetailCustomizations/Private/DialogueWaveWidgets.cpp
@@ -906,6 +906,9 @@ void SDialogueContextHeaderWidget::Construct( const FArguments& InArgs, const TS
 		];
 	}
 }
+END_SLATE_FUNCTION_BUILD_OPTIMIZATION
+
+
 void SDialogueContextHeaderWidget::Tick( const FGeometry& AllottedGeometry, const double InCurrentTime, const float InDeltaTime )
 {
 	if( !IsSpeakerValid() )
diff --git a/Engine/Source/Editor/GameProjectGeneration/Private/SGameProjectDialog.cpp b/Engine/Source/Editor/GameProjectGeneration/Private/SGameProjectDialog.cpp
--- a/Engine/Source/Editor/GameProjectGeneration/Private/SGameProjectDialog.cpp
+++ b/Engine/Source/Editor/GameProjectGeneration/Private/SGameProjectDialog.cpp
@@ -188,6 +188,8 @@ void SGameProjectDialog::Construct( const FArguments& InArgs )
 		check(false);
 	}
 }
+END_SLATE_FUNCTION_BUILD_OPTIMIZATION
+
 
 EActiveTimerReturnType SGameProjectDialog::TriggerFadeInPostConstruct(double InCurrentTime, float InDeltaTime)
 {
diff --git a/Engine/Source/Editor/GameProjectGeneration/Private/SNewClassDialog.cpp b/Engine/Source/Editor/GameProjectGeneration/Private/SNewClassDialog.cpp
--- a/Engine/Source/Editor/GameProjectGeneration/Private/SNewClassDialog.cpp
+++ b/Engine/Source/Editor/GameProjectGeneration/Private/SNewClassDialog.cpp
@@ -791,6 +791,9 @@ void SNewClassDialog::Construct( const FArguments& InArgs )
 		ParentClassListView->SetSelection(ParentClassItemsSource[0], ESelectInfo::Direct);
 	}
 }
+END_SLATE_FUNCTION_BUILD_OPTIMIZATION
+
+
 void SNewClassDialog::Tick( const FGeometry& AllottedGeometry, const double InCurrentTime, const float InDeltaTime )
 {
 	// Every few seconds, the class name/path is checked for validity in case the disk contents changed and the location is now valid or invalid.
diff --git a/Engine/Source/Editor/MeshPaint/Private/SMeshPaint.cpp b/Engine/Source/Editor/MeshPaint/Private/SMeshPaint.cpp
--- a/Engine/Source/Editor/MeshPaint/Private/SMeshPaint.cpp
+++ b/Engine/Source/Editor/MeshPaint/Private/SMeshPaint.cpp
@@ -2096,6 +2096,7 @@ void SMeshPaint::Construct(const FArguments& InArgs, TSharedRef<FMeshPaintToolKi
 		]
 	];
 }
+END_SLATE_FUNCTION_BUILD_OPTIMIZATION
 
 void SMeshPaint::Tick( const FGeometry& AllottedGeometry, const double InCurrentTime, const float InDeltaTime )
 {
diff --git a/Engine/Source/Editor/UndoHistory/Private/Widgets/SUndoHistory.cpp b/Engine/Source/Editor/UndoHistory/Private/Widgets/SUndoHistory.cpp
--- a/Engine/Source/Editor/UndoHistory/Private/Widgets/SUndoHistory.cpp
+++ b/Engine/Source/Editor/UndoHistory/Private/Widgets/SUndoHistory.cpp
@@ -75,6 +75,8 @@ void SUndoHistory::Construct( const FArguments& InArgs )
 
 	ReloadUndoList();
 }
+END_SLATE_FUNCTION_BUILD_OPTIMIZATION
+
 void SUndoHistory::Tick( const FGeometry& AllottedGeometry, const double InCurrentTime, const float InDeltaTime )
 {
 	// reload the transaction list if necessary
diff --git a/Engine/Source/Runtime/Core/Public/Templates/SharedPointer.h b/Engine/Source/Runtime/Core/Public/Templates/SharedPointer.h
--- a/Engine/Source/Runtime/Core/Public/Templates/SharedPointer.h
+++ b/Engine/Source/Runtime/Core/Public/Templates/SharedPointer.h
@@ -489,7 +489,7 @@ class TSharedPtr
 
 	// TSharedPtr of one mode to a type which has a TSharedFromThis only of another mode is illegal.
 	// A type which does not inherit TSharedFromThis at all is ok.
-	static_assert(ObjectTypeHasSameModeSharedFromThis || !ObjectTypeHasOppositeModeSharedFromThis, "You cannot use a TSharedPtr of one mode with a type which inherits TSharedFromThis of another mode.");
+	static_assert( TSharedPtr::ObjectTypeHasSameModeSharedFromThis || !TSharedPtr::ObjectTypeHasOppositeModeSharedFromThis, "You cannot use a TSharedPtr of one mode with a type which inherits TSharedFromThis of another mode.");
 
 public:
 
diff --git a/Engine/Source/Runtime/Core/Public/Windows/WindowsPlatformCompilerSetup.h b/Engine/Source/Runtime/Core/Public/Windows/WindowsPlatformCompilerSetup.h
--- a/Engine/Source/Runtime/Core/Public/Windows/WindowsPlatformCompilerSetup.h
+++ b/Engine/Source/Runtime/Core/Public/Windows/WindowsPlatformCompilerSetup.h
@@ -176,6 +176,7 @@ static_assert(_MSC_VER >= 1800, "Visual Studio 2013 or later is required to comp
 #pragma warning(disable : 4180) // qualifier applied to function type has no meaning; ignored
 #pragma warning(disable : 4121) // 'symbol' : alignment of a member was sensitive to packing
 #pragma warning(disable : 4345) // behavior change: an object of POD type constructed with an initializer of the form () will be default-initialized
+#pragma warning(disable : 4464) // relative include path contains '..'                                            
 
 #if UE_BUILD_DEBUG
 // xstring.h causes this warning in debug builds
