Copying //UE4/Dev-Sequencer to //UE4/Dev-Main (Source: //UE4/Dev-Sequencer @ 3178529)

diff --git a/Engine/Plugins/Messaging/MessagingDebugger/Source/MessagingDebugger/Private/Widgets/History/SMessagingHistoryTableRow.h b/Engine/Plugins/Messaging/MessagingDebugger/Source/MessagingDebugger/Private/Widgets/History/SMessagingHistoryTableRow.h
--- a/Engine/Plugins/Messaging/MessagingDebugger/Source/MessagingDebugger/Private/Widgets/History/SMessagingHistoryTableRow.h
+++ b/Engine/Plugins/Messaging/MessagingDebugger/Source/MessagingDebugger/Private/Widgets/History/SMessagingHistoryTableRow.h
@@ -3,20 +3,20 @@
 #pragma once
 
 
-#define LOCTEXT_NAMESPACE "SMessagingHistoryTableRow"
+struct FMessageTracerMessageInfo;
 
 
 /**
  * Implements a row widget for the message history list.
  */
 class SMessagingHistoryTableRow
-	: public SMultiColumnTableRow<FMessageTracerMessageInfoPtr>
+	: public SMultiColumnTableRow<TSharedPtr<FMessageTracerMessageInfo>>
 {
 public:
 
 	SLATE_BEGIN_ARGS(SMessagingHistoryTableRow) { }
 		SLATE_ATTRIBUTE(FText, HighlightText)
-		SLATE_ARGUMENT(FMessageTracerMessageInfoPtr, MessageInfo)
+		SLATE_ARGUMENT(TSharedPtr<FMessageTracerMessageInfo>, MessageInfo)
 		SLATE_ARGUMENT(TSharedPtr<ISlateStyle>, Style)
 	SLATE_END_ARGS()
 
@@ -28,191 +28,19 @@ public:
 	 * @param InArgs The construction arguments.
 	 * @param InOwnerTableView The table view that owns this row.
 	 */
-	void Construct( const FArguments& InArgs, const TSharedRef<STableViewBase>& InOwnerTableView )
-	{
-		check(InArgs._MessageInfo.IsValid());
-		check(InArgs._Style.IsValid());
-
-		MaxDispatchLatency = -1.0;
-		MaxHandlingTime = -1.0;
-		HighlightText = InArgs._HighlightText;
-		MessageInfo = InArgs._MessageInfo;
-		Style = InArgs._Style;
-
-		SMultiColumnTableRow<FMessageTracerMessageInfoPtr>::Construct(FSuperRowType::FArguments(), InOwnerTableView);
-	}
+	void Construct(const FArguments& InArgs, const TSharedRef<STableViewBase>& InOwnerTableView);
 
 public:
 
-	// SWidget overrides
-
-	virtual void Tick( const FGeometry& AllottedGeometry, const double InCurrentTime, const float InDeltaTime ) override
-	{
-		MaxDispatchLatency = -1.0;
-		MaxHandlingTime = -1.0;
-
-		for (TMap<FMessageTracerEndpointInfoPtr, FMessageTracerDispatchStatePtr>::TConstIterator It(MessageInfo->DispatchStates); It; ++It)
-		{
-			const FMessageTracerDispatchStatePtr& DispatchState = It.Value();
-
-			if (MessageInfo->TimeRouted > 0.0)
-			{
-				MaxDispatchLatency = FMath::Max(MaxDispatchLatency, DispatchState->DispatchLatency);
-			}
+	//~ SWidget overrides
 
-			if (DispatchState->TimeHandled > 0.0)
-			{
-				MaxHandlingTime = FMath::Max(MaxHandlingTime, DispatchState->TimeHandled - DispatchState->TimeDispatched);
-			}
-		}
-	}
+	virtual void Tick(const FGeometry& AllottedGeometry, const double InCurrentTime, const float InDeltaTime) override;
 
 public:
 
-	// SMultiColumnTableRow interface
-
-	BEGIN_SLATE_FUNCTION_BUILD_OPTIMIZATION
-	virtual TSharedRef<SWidget> GenerateWidgetForColumn( const FName& ColumnName ) override
-	{
-		if (ColumnName == "DispatchLatency")
-		{
-			return SNew(SBox)
-				.Padding(FMargin(4.0f, 0.0f))
-				[
-					SNew(STextBlock)
-						.ColorAndOpacity(this, &SMessagingHistoryTableRow::HandleTextColorAndOpacity)
-						.Text(this, &SMessagingHistoryTableRow::HandleDispatchLatencyText)
-				];
-		}
-		else if (ColumnName == "Flag")
-		{
-			return SNew(SBox)
-				.ToolTipText(LOCTEXT("DeadMessageTooltip", "No valid recipients (dead letter)"))
-				[
-					SNew(SImage)
-						.Image(this, &SMessagingHistoryTableRow::HandleFlagImage)
-				];
-		}
-		else if (ColumnName == "HandleTime")
-		{
-			return SNew(SBox)
-				.Padding(FMargin(4.0f, 0.0f))
-				[
-					SNew(STextBlock)
-						.ColorAndOpacity(this, &SMessagingHistoryTableRow::HandleTextColorAndOpacity)
-						.Text(this, &SMessagingHistoryTableRow::HandleHandlingTimeText)
-				];
-		}
-		else if (ColumnName == "MessageType")
-		{
-			return SNew(SBox)
-				.Padding(FMargin(4.0f, 0.0f))
-				[
-					SNew(STextBlock)
-						.ColorAndOpacity(this, &SMessagingHistoryTableRow::HandleTextColorAndOpacity)
-						.HighlightText(HighlightText)
-						.Text(FText::FromName(MessageInfo->Context->GetMessageType()))
-				];
-		}
-		else if (ColumnName == "Recipients")
-		{
-			return SNew(SBox)
-				.Padding(FMargin(4.0f, 0.0f))
-				.ToolTipText(LOCTEXT("LocalRemoteRecipients", "Local/Remote Recipients"))
-				[
-					SNew(STextBlock)
-						.ColorAndOpacity(this, &SMessagingHistoryTableRow::HandleTextColorAndOpacity)
-						.Text(this, &SMessagingHistoryTableRow::HandleRecipientsText)
-				];
-		}
-		else if (ColumnName == "RouteLatency")
-		{
-			return SNew(SBox)
-				.Padding(FMargin(4.0f, 0.0f))
-				[
-					SNew(STextBlock)
-						.ColorAndOpacity(this, &SMessagingHistoryTableRow::HandleRouteLatencyColorAndOpacity)
-						.Text(this, &SMessagingHistoryTableRow::HandleRouteLatencyText)
-				];
-		}
-		else if (ColumnName == "Scope")
-		{
-			FText Text;
-			FText ToolTipText;
-
-			if (MessageInfo->Context.IsValid())
-			{
-				int32 NumRecipients = MessageInfo->Context->GetRecipients().Num();
-
-				if (MessageInfo->Context->IsForwarded())
-				{
-					// forwarded message
-					FText ScopeText = ScopeToText(MessageInfo->Context->GetScope());
-					Text = FText::Format(LOCTEXT("ForwardedMessageTextFormat", "F - {0}"), ScopeText);
-
-					if (NumRecipients > 0)
-					{
-						ToolTipText = FText::Format(LOCTEXT("ForwardedSentMessageToolTipTextFormat", "This message was forwarded directly to {0} recipients"), FText::AsNumber(NumRecipients));
-					}
-					else
-					{
-						ToolTipText = FText::Format(LOCTEXT("ForwardedPublishedMessageToolTipTextFormat", "This message was forwarded to all subscribed recipients in scope '{0}'"), ScopeText);
-					}				
-				}
-				else if (NumRecipients == 0)
-				{
-					// published message
-					FText ScopeText = ScopeToText(MessageInfo->Context->GetScope());
-
-					Text = FText::Format(LOCTEXT("PublishedMessageTextFormat", "P - {0}"), ScopeText);
-					ToolTipText = FText::Format(LOCTEXT("PublishedMessageToolTipTextFormat", "This message was published to all subscribed recipients in scope '{0}'"), ScopeText);
-				}
-				else
-				{
-					// sent message
-					Text = LOCTEXT("SentMessageTextFormat", "S");
-					ToolTipText = FText::Format(LOCTEXT("SentMessageToolTipTextFormat", "This message was sent directly to {0} recipients"), FText::AsNumber(NumRecipients));
-				}
-			}
+	//~ SMultiColumnTableRow interface
 
-			return SNew(SBox)
-				.ToolTipText(ToolTipText)
-				[
-					SNew(STextBlock)
-						.ColorAndOpacity(this, &SMessagingHistoryTableRow::HandleTextColorAndOpacity)
-						.Text(Text)
-				];
-		}
-		else if (ColumnName == "Sender")
-		{
-			return SNew(SBox)
-				.Padding(FMargin(4.0f, 0.0f))
-				[
-					SNew(STextBlock)
-						.ColorAndOpacity(this, &SMessagingHistoryTableRow::HandleTextColorAndOpacity)
-						.HighlightText(HighlightText)
-						.Text(this, &SMessagingHistoryTableRow::HandleSenderText)
-				];
-		}
-		else if (ColumnName == "TimeSent")
-		{
-			FNumberFormattingOptions NumberFormattingOptions;
-			NumberFormattingOptions.MaximumFractionalDigits = 3;
-			NumberFormattingOptions.MinimumFractionalDigits = 3;
-
-			return SNew(SBox)
-				.Padding(FMargin(4.0f, 0.0f))
-				[
-					SNew(STextBlock)
-						.ColorAndOpacity(this, &SMessagingHistoryTableRow::HandleTextColorAndOpacity)
-						.HighlightText(HighlightText)
-						.Text(FText::AsNumber(MessageInfo->TimeSent - GStartTime, &NumberFormattingOptions))
-				];
-		}
-
-		return SNullWidget::NullWidget;
-	}
-	END_SLATE_FUNCTION_BUILD_OPTIMIZATION
+	virtual TSharedRef<SWidget> GenerateWidgetForColumn(const FName& ColumnName) override;
 
 protected:
 
@@ -221,30 +49,7 @@ protected:
 	 *
 	 * @return The string representation.
 	 */
-	FText ScopeToText( EMessageScope Scope ) const
-	{
-		switch (Scope)
-		{
-		case EMessageScope::Thread:
-			return LOCTEXT("ScopeThread", "Thread");
-			break;
-
-		case EMessageScope::Process:
-			return LOCTEXT("ScopeProcess", "Process");
-			break;
-
-		case EMessageScope::Network:
-			return LOCTEXT("ScopeNetwork", "Network");
-			break;
-
-		case EMessageScope::All:
-			return LOCTEXT("ScopeAll", "All");
-			break;
-
-		default:
-			return LOCTEXT("ScopeUnknown", "Unknown");
-		}
-	}
+	FText ScopeToText(EMessageScope Scope) const;
 
 	/**
 	 * Converts the given time span in seconds to a human readable string.
@@ -254,161 +59,26 @@ protected:
 	 *
 	 * @todo gmp: refactor this into FText::AsTimespan or something like that
 	 */
-	FText TimespanToReadableText( double Seconds ) const
-	{
-		if (Seconds < 0.0)
-		{
-			return LOCTEXT("Zero Length Timespan", "-");
-		}
-
-		FNumberFormattingOptions Options;
-		Options.MinimumFractionalDigits = 1;
-		Options.MaximumFractionalDigits = 1;
-
-		if (Seconds < 0.0001)
-		{
-			return FText::Format(LOCTEXT("Seconds < 0.0001 Length Timespan", "{0} us"), FText::AsNumber(Seconds * 1000000, &Options));
-		}
-
-		if (Seconds < 0.1)
-		{
-			return FText::Format(LOCTEXT("Seconds < 0.1 Length Timespan", "{0} ms"), FText::AsNumber(Seconds * 1000, &Options));
-		}
-
-		if (Seconds < 60.0)
-		{
-			return FText::Format(LOCTEXT("Seconds < 60.0 Length Timespan", "{0} s"), FText::AsNumber(Seconds, &Options));
-		}
-		
-		return LOCTEXT("> 1 minute Length Timespan", "> 1 min");
-	}
+	FText TimespanToReadableText(double Seconds) const;
 
 private:
 
 	/** Callback for getting the text color of the DispatchLatency column. */
-	FSlateColor HandleDispatchLatencyColorAndOpacity() const
-	{
-		if (MaxDispatchLatency >= 0.01)
-		{
-			return FLinearColor::Red;
-		}
+	FSlateColor HandleDispatchLatencyColorAndOpacity() const;
 
-		if (MaxDispatchLatency >= 0.001)
-		{
-			return FLinearColor(1.0f, 1.0f, 0.0f);
-		}
-
-		if (MaxDispatchLatency >= 0.0001)
-		{
-			return FLinearColor::Yellow;
-		}
-
-		return FSlateColor::UseForeground();
-	}
-
-	/** Callback for getting the text in the DispatchLatency column. */
-	FText HandleDispatchLatencyText() const
-	{
-		return TimespanToReadableText(MaxDispatchLatency);
-	}
-
-	const FSlateBrush* HandleFlagImage() const
-	{
-		if ((MessageInfo->TimeRouted > 0.0) && (MessageInfo->DispatchStates.Num() == 0))
-		{
-			return Style->GetBrush("DeadMessage");
-		}
-
-		return NULL;
-	}
-
-	/** Callback for getting the text in the HandlingTime column. */
-	FText HandleHandlingTimeText() const
-	{
-		return TimespanToReadableText(MaxHandlingTime);
-	}
+	const FSlateBrush* HandleFlagImage() const;
 
 	/** Callback for getting the text in the Recipients column. */
-	FText HandleRecipientsText() const
-	{
-		int32 LocalRecipients = 0;
-		int32 RemoteRecipients = 0;
-
-		for (auto It = MessageInfo->DispatchStates.CreateConstIterator(); It; ++It)
-		{
-			FMessageTracerEndpointInfoPtr EndpointInfo = It.Value()->EndpointInfo;
-
-			if (EndpointInfo.IsValid())
-			{
-				if (EndpointInfo->Remote)
-				{
-					++RemoteRecipients;
-				}
-				else
-				{
-					++LocalRecipients;
-				}
-			}
-		}
-
-		return FText::Format(LOCTEXT("RecipientsTextFormat", "{0} / {1}"), FText::AsNumber(LocalRecipients), FText::AsNumber(RemoteRecipients));
-	}
+	FText HandleRecipientsText() const;
 
 	/** Callback for getting the text color of the RouteLatency column. */
-	FSlateColor HandleRouteLatencyColorAndOpacity() const
-	{
-		double RouteLatency = MessageInfo->TimeRouted - MessageInfo->TimeSent;
-
-		if (RouteLatency >= 0.01)
-		{
-			return FLinearColor::Red;
-		}
-
-		if (RouteLatency >= 0.001)
-		{
-			return FLinearColor(1.0f, 1.0f, 0.0f);
-		}
-
-		if (RouteLatency >= 0.0001)
-		{
-			return FLinearColor::Yellow;
-		}
-
-		if (MessageInfo->TimeRouted == 0.0)
-		{
-			return FSlateColor::UseSubduedForeground();
-		}
-
-		return FSlateColor::UseForeground();
-	}
+	FSlateColor HandleRouteLatencyColorAndOpacity() const;
 
 	/** Callback for getting the text in the RouteLatency column. */
-	FText HandleRouteLatencyText() const
-	{
-		if (MessageInfo->TimeRouted > 0.0)
-		{
-			return TimespanToReadableText(MessageInfo->TimeRouted - MessageInfo->TimeSent);
-		}
-
-		return LOCTEXT("RouteLatencyPending", "Pending");
-	}
-
-	/** Callback for getting the text in the Sender column. */
-	FText HandleSenderText() const
-	{
-		return FText::FromString(MessageInfo->SenderInfo->Name.ToString());
-	}
+	FText HandleRouteLatencyText() const;
 
 	/** Callback for getting the text color of various columns. */
-	FSlateColor HandleTextColorAndOpacity() const
-	{
-		if (MessageInfo->TimeRouted == 0.0)
-		{
-			return FSlateColor::UseSubduedForeground();
-		}
-
-		return FSlateColor::UseForeground();
-	}
+	FSlateColor HandleTextColorAndOpacity() const;
 
 private:
 
@@ -416,7 +86,7 @@ private:
 	TAttribute<FText> HighlightText;
 
 	/** Holds message's debug information. */
-	FMessageTracerMessageInfoPtr MessageInfo;
+	TSharedPtr<FMessageTracerMessageInfo> MessageInfo;
 
 	/** Holds the maximum dispatch latency. */
 	double MaxDispatchLatency;
@@ -427,6 +97,3 @@ private:
 	/** Holds the widget's visual style. */
 	TSharedPtr<ISlateStyle> Style;
 };
-
-
-#undef LOCTEXT_NAMESPACE
diff --git a/Engine/Plugins/Messaging/MessagingDebugger/Source/MessagingDebugger/Private/Widgets/Interceptors/SMessagingInterceptors.cpp b/Engine/Plugins/Messaging/MessagingDebugger/Source/MessagingDebugger/Private/Widgets/Interceptors/SMessagingInterceptors.cpp
--- a/Engine/Plugins/Messaging/MessagingDebugger/Source/MessagingDebugger/Private/Widgets/Interceptors/SMessagingInterceptors.cpp
+++ b/Engine/Plugins/Messaging/MessagingDebugger/Source/MessagingDebugger/Private/Widgets/Interceptors/SMessagingInterceptors.cpp
@@ -1,6 +1,8 @@
 // Copyright 1998-2016 Epic Games, Inc. All Rights Reserved.
 
-#include "MessagingDebuggerPrivatePCH.h"
+#include "MessagingDebuggerPCH.h"
+#include "MessagingDebuggerModel.h"
+#include "SMessagingInterceptors.h"
 #include "SMessagingInterceptorsTableRow.h"
 
 
@@ -11,7 +13,7 @@
  *****************************************************************************/
 
 BEGIN_SLATE_FUNCTION_BUILD_OPTIMIZATION
-void SMessagingInterceptors::Construct(const FArguments& InArgs, const FMessagingDebuggerModelRef& InModel, const TSharedRef<ISlateStyle>& InStyle, const IMessageTracerRef& InTracer)
+void SMessagingInterceptors::Construct(const FArguments& InArgs, const TSharedRef<FMessagingDebuggerModel>& InModel, const TSharedRef<ISlateStyle>& InStyle, const TSharedRef<IMessageTracer, ESPMode::ThreadSafe>& InTracer)
 {
 	Model = InModel;
 	Style = InStyle;
@@ -30,7 +32,7 @@ void SMessagingInterceptors::Construct(const FArguments& InArgs, const FMessagin
 					.Padding(0.0f)
 					[
 						// interceptor list
-						SAssignNew(InterceptorListView, SListView<FMessageTracerInterceptorInfoPtr>)
+						SAssignNew(InterceptorListView, SListView<TSharedPtr<FMessageTracerInterceptorInfo>>)
 							.ItemHeight(24.0f)
 							.ListItemsSource(&InterceptorList)
 							.SelectionMode(ESelectionMode::None)
@@ -74,7 +76,7 @@ void SMessagingInterceptors::ReloadInterceptorList()
 /* SMessagingInterceptors callbacks
  *****************************************************************************/
 
-TSharedRef<ITableRow> SMessagingInterceptors::HandleInterceptorListGenerateRow(FMessageTracerInterceptorInfoPtr InterceptorInfo, const TSharedRef<STableViewBase>& OwnerTable)
+TSharedRef<ITableRow> SMessagingInterceptors::HandleInterceptorListGenerateRow(TSharedPtr<FMessageTracerInterceptorInfo> InterceptorInfo, const TSharedRef<STableViewBase>& OwnerTable)
 {
 	return SNew(SMessagingInterceptorsTableRow, OwnerTable, Model.ToSharedRef())
 		.InterceptorInfo(InterceptorInfo)
diff --git a/Engine/Plugins/Messaging/MessagingDebugger/Source/MessagingDebugger/Private/Widgets/Interceptors/SMessagingInterceptors.h b/Engine/Plugins/Messaging/MessagingDebugger/Source/MessagingDebugger/Private/Widgets/Interceptors/SMessagingInterceptors.h
--- a/Engine/Plugins/Messaging/MessagingDebugger/Source/MessagingDebugger/Private/Widgets/Interceptors/SMessagingInterceptors.h
+++ b/Engine/Plugins/Messaging/MessagingDebugger/Source/MessagingDebugger/Private/Widgets/Interceptors/SMessagingInterceptors.h
@@ -3,6 +3,11 @@
 #pragma once
 
 
+class FMessagingDebuggerModel;
+struct FMessageTracerInterceptorInfo;
+class IMessageTracer;
+
+
 /**
  * Implements the message interceptors panel.
  */
@@ -24,7 +29,11 @@ public:
 	 * @param InStyle The visual style to use for this widget.
 	 * @param InTracer The message tracer.
 	 */
-	void Construct(const FArguments& InArgs, const FMessagingDebuggerModelRef& InModel, const TSharedRef<ISlateStyle>& InStyle, const IMessageTracerRef& InTracer);
+	void Construct(
+		const FArguments& InArgs,
+		const TSharedRef<FMessagingDebuggerModel>& InModel,
+		const TSharedRef<ISlateStyle>& InStyle,
+		const TSharedRef<IMessageTracer, ESPMode::ThreadSafe>& InTracer);
 
 protected:
 
@@ -34,22 +43,22 @@ protected:
 private:
 
 	/** Handles generating a row widget for the interceptor list view. */
-	TSharedRef<ITableRow> HandleInterceptorListGenerateRow(FMessageTracerInterceptorInfoPtr InterceptorInfo, const TSharedRef<STableViewBase>& OwnerTable);
+	TSharedRef<ITableRow> HandleInterceptorListGenerateRow(TSharedPtr<FMessageTracerInterceptorInfo> InterceptorInfo, const TSharedRef<STableViewBase>& OwnerTable);
 
 private:
 
 	/** Holds the filtered list of interceptors. */
-	TArray<FMessageTracerInterceptorInfoPtr> InterceptorList;
+	TArray<TSharedPtr<FMessageTracerInterceptorInfo>> InterceptorList;
 
 	/** Holds the interceptor list view. */
-	TSharedPtr<SListView<FMessageTracerInterceptorInfoPtr>> InterceptorListView;
+	TSharedPtr<SListView<TSharedPtr<FMessageTracerInterceptorInfo>>> InterceptorListView;
 
 	/** Holds a pointer to the view model. */
-	FMessagingDebuggerModelPtr Model;
+	TSharedPtr<FMessagingDebuggerModel> Model;
 
 	/** Holds the widget's visual style. */
 	TSharedPtr<ISlateStyle> Style;
 
 	/** Holds a pointer to the message bus tracer. */
-	IMessageTracerPtr Tracer;
+	TSharedPtr<IMessageTracer, ESPMode::ThreadSafe> Tracer;
 };
diff --git a/Engine/Plugins/Messaging/MessagingDebugger/Source/MessagingDebugger/Private/Widgets/Interceptors/SMessagingInterceptorsTableRow.h b/Engine/Plugins/Messaging/MessagingDebugger/Source/MessagingDebugger/Private/Widgets/Interceptors/SMessagingInterceptorsTableRow.h
--- a/Engine/Plugins/Messaging/MessagingDebugger/Source/MessagingDebugger/Private/Widgets/Interceptors/SMessagingInterceptorsTableRow.h
+++ b/Engine/Plugins/Messaging/MessagingDebugger/Source/MessagingDebugger/Private/Widgets/Interceptors/SMessagingInterceptorsTableRow.h
@@ -3,27 +3,20 @@
 #pragma once
 
 
-#define LOCTEXT_NAMESPACE "SMessagingInterceptorTableRow"
-
-
-namespace MessagingInterceptorTableRow
-{
-	static const FNumberFormattingOptions TimeRegisteredFormattingOptions = FNumberFormattingOptions()
-		.SetMinimumFractionalDigits(5)
-		.SetMaximumFractionalDigits(5);
-}
+class FMessagingDebuggerModel;
+struct FMessageTracerInterceptorInfo;
 
 
 /**
  * Implements a row widget for the interceptors list.
  */
 class SMessagingInterceptorsTableRow
-	: public SMultiColumnTableRow<FMessageTracerInterceptorInfoPtr>
+	: public SMultiColumnTableRow<TSharedPtr<FMessageTracerInterceptorInfo>>
 {
 public:
 
 	SLATE_BEGIN_ARGS(SMessagingInterceptorsTableRow) { }
-		SLATE_ARGUMENT(FMessageTracerInterceptorInfoPtr, InterceptorInfo)
+		SLATE_ARGUMENT(TSharedPtr<FMessageTracerInterceptorInfo>, InterceptorInfo)
 		SLATE_ARGUMENT(TSharedPtr<ISlateStyle>, Style)
 	SLATE_END_ARGS()
 
@@ -35,84 +28,27 @@ public:
 	 * @param InArgs The construction arguments.
 	 * @param InOwnerTableView The table view that owns this row.
 	 */
-	void Construct(const FArguments& InArgs, const TSharedRef<STableViewBase>& InOwnerTableView, const FMessagingDebuggerModelRef& InModel)
-	{
-		check(InArgs._Style.IsValid());
-		check(InArgs._InterceptorInfo.IsValid());
-
-		InterceptorInfo = InArgs._InterceptorInfo;
-		Model = InModel;
-		Style = InArgs._Style;
-
-		SMultiColumnTableRow<FMessageTracerInterceptorInfoPtr>::Construct(FSuperRowType::FArguments(), InOwnerTableView);
-	}
+	void Construct(const FArguments& InArgs, const TSharedRef<STableViewBase>& InOwnerTableView, const TSharedRef<FMessagingDebuggerModel>& InModel);
 
 public:
 
-	// SMultiColumnTableRow interface
-
-	BEGIN_SLATE_FUNCTION_BUILD_OPTIMIZATION
-	virtual TSharedRef<SWidget> GenerateWidgetForColumn(const FName& ColumnName) override
-	{
-		if (ColumnName == "Name")
-		{
-			return SNew(SBox)
-				.Padding(FMargin(4.0f, 0.0f))
-				.VAlign(VAlign_Center)
-				[
-					SNew(STextBlock)
-						.Text(FText::FromString(InterceptorInfo->Name.ToString()))
-				];
-		}
-		else if (ColumnName == "TimeRegistered")
-		{
-			return SNew(SBox)
-				.Padding(FMargin(4.0f, 0.0f))
-				.VAlign(VAlign_Center)
-				[
-					SNew(STextBlock)
-						.Text(FText::AsNumber(InterceptorInfo->TimeRegistered - GStartTime, &MessagingInterceptorTableRow::TimeRegisteredFormattingOptions))
-				];
-		}
-		else if (ColumnName == "TimeUnregistered")
-		{
-			return SNew(SBox)
-				.Padding(FMargin(4.0f, 0.0f))
-				.VAlign(VAlign_Center)
-				[
-					SNew(STextBlock)
-						.Text(this, &SMessagingInterceptorsTableRow::HandleTimeUnregisteredText)
-				];
-		}
-
-		return SNullWidget::NullWidget;
-	}
-	END_SLATE_FUNCTION_BUILD_OPTIMIZATION
+	//~ SMultiColumnTableRow interface
+
+	virtual TSharedRef<SWidget> GenerateWidgetForColumn(const FName& ColumnName) override;
 
 private:
 
 	/** Callback for getting the timestamp at which the address was unregistered. */
-	FText HandleTimeUnregisteredText() const
-	{
-		if (InterceptorInfo->TimeUnregistered > 0.0)
-		{
-			return FText::AsNumber(InterceptorInfo->TimeUnregistered - GStartTime, &MessagingInterceptorTableRow::TimeRegisteredFormattingOptions);
-		}
-
-		return LOCTEXT("Never", "Never");
-	}
+	FText HandleTimeUnregisteredText() const;
 
 private:
 
 	/** Holds the interceptor information. */
-	FMessageTracerInterceptorInfoPtr InterceptorInfo;
+	TSharedPtr<FMessageTracerInterceptorInfo> InterceptorInfo;
 
 	/** Holds a pointer to the view model. */
-	FMessagingDebuggerModelPtr Model;
+	TSharedPtr<FMessagingDebuggerModel> Model;
 
 	/** Holds the widget's visual style. */
 	TSharedPtr<ISlateStyle> Style;
 };
-
-
-#undef LOCTEXT_NAMESPACE
diff --git a/Engine/Plugins/Messaging/MessagingDebugger/Source/MessagingDebugger/Private/Widgets/MessageData/SMessagingMessageData.cpp b/Engine/Plugins/Messaging/MessagingDebugger/Source/MessagingDebugger/Private/Widgets/MessageData/SMessagingMessageData.cpp
--- a/Engine/Plugins/Messaging/MessagingDebugger/Source/MessagingDebugger/Private/Widgets/MessageData/SMessagingMessageData.cpp
+++ b/Engine/Plugins/Messaging/MessagingDebugger/Source/MessagingDebugger/Private/Widgets/MessageData/SMessagingMessageData.cpp
@@ -1,15 +1,8 @@
 // Copyright 1998-2016 Epic Games, Inc. All Rights Reserved.
 
-#include "MessagingDebuggerPrivatePCH.h"
-#include "Json.h"
-#include "JsonStructSerializerBackend.h"
-#include "StructSerializer.h"
-
-#if WITH_EDITOR
-	#include "IDetailsView.h"
-	#include "IStructureDetailsView.h"
-	#include "PropertyEditorModule.h"
-#endif
+#include "MessagingDebuggerPCH.h"
+#include "MessagingDebuggerModel.h"
+#include "SMessagingMessageData.h"
 
 
 #define LOCTEXT_NAMESPACE "SMessagingMessageData"
@@ -30,7 +23,7 @@ SMessagingMessageData::~SMessagingMessageData()
 /* SMessagingMessageData interface
  *****************************************************************************/
 
-void SMessagingMessageData::Construct( const FArguments& InArgs, const FMessagingDebuggerModelRef& InModel, const TSharedRef<ISlateStyle>& InStyle )
+void SMessagingMessageData::Construct(const FArguments& InArgs, const TSharedRef<FMessagingDebuggerModel>& InModel, const TSharedRef<ISlateStyle>& InStyle)
 {
 	Model = InModel;
 	Style = InStyle;
@@ -88,7 +81,7 @@ void SMessagingMessageData::Construct( const FArguments& InArgs, const FMessagin
 /* FNotifyHook interface
  *****************************************************************************/
 
-void SMessagingMessageData::NotifyPostChange( const FPropertyChangedEvent& PropertyChangedEvent, class FEditPropertyChain* PropertyThatChanged )
+void SMessagingMessageData::NotifyPostChange(const FPropertyChangedEvent& PropertyChangedEvent, class FEditPropertyChain* PropertyThatChanged)
 {
 }
 
@@ -98,7 +91,7 @@ void SMessagingMessageData::NotifyPostChange( const FPropertyChangedEvent& Prope
 
 bool SMessagingMessageData::HandleDetailsViewIsPropertyEditable() const
 {
-	FMessageTracerMessageInfoPtr SelectedMessage = Model->GetSelectedMessage();
+	TSharedPtr<FMessageTracerMessageInfo> SelectedMessage = Model->GetSelectedMessage();
 
 	if (!SelectedMessage.IsValid() || !SelectedMessage->Context.IsValid())
 	{
@@ -122,7 +115,7 @@ EVisibility SMessagingMessageData::HandleDetailsViewVisibility() const
 
 void SMessagingMessageData::HandleModelSelectedMessageChanged()
 {
-	FMessageTracerMessageInfoPtr SelectedMessage = Model->GetSelectedMessage();
+	TSharedPtr<FMessageTracerMessageInfo> SelectedMessage = Model->GetSelectedMessage();
 
 	if (SelectedMessage.IsValid() && SelectedMessage->Context.IsValid())
 	{
diff --git a/Engine/Plugins/Messaging/MessagingDebugger/Source/MessagingDebugger/Private/Widgets/MessageData/SMessagingMessageData.h b/Engine/Plugins/Messaging/MessagingDebugger/Source/MessagingDebugger/Private/Widgets/MessageData/SMessagingMessageData.h
--- a/Engine/Plugins/Messaging/MessagingDebugger/Source/MessagingDebugger/Private/Widgets/MessageData/SMessagingMessageData.h
+++ b/Engine/Plugins/Messaging/MessagingDebugger/Source/MessagingDebugger/Private/Widgets/MessageData/SMessagingMessageData.h
@@ -3,6 +3,7 @@
 #pragma once
 
 
+class FMessagingDebuggerModel;
 class IStructureDetailsView;
 
 
@@ -32,11 +33,11 @@ public:
 	 * @param InModel The view model to use.
 	 * @param InStyle The visual style to use for this widget.
 	 */
-	void Construct(const FArguments& InArgs, const FMessagingDebuggerModelRef& InModel, const TSharedRef<ISlateStyle>& InStyle);
+	void Construct(const FArguments& InArgs, const TSharedRef<FMessagingDebuggerModel>& InModel, const TSharedRef<ISlateStyle>& InStyle);
 
 public:
 
-	// FNotifyHook interface
+	//~ FNotifyHook interface
 
 	virtual void NotifyPostChange(const FPropertyChangedEvent& PropertyChangedEvent, class FEditPropertyChain* PropertyThatChanged) override;
 
@@ -62,7 +63,7 @@ private:
 #endif //WITH_EDITOR
 
 	/** Holds a pointer to the view model. */
-	FMessagingDebuggerModelPtr Model;
+	TSharedPtr<FMessagingDebuggerModel> Model;
 
 	/** Holds the widget's visual style. */
 	TSharedPtr<ISlateStyle> Style;
diff --git a/Engine/Plugins/Messaging/MessagingDebugger/Source/MessagingDebugger/Private/Widgets/MessageDetails/SMessagingDispatchStateTableRow.h b/Engine/Plugins/Messaging/MessagingDebugger/Source/MessagingDebugger/Private/Widgets/MessageDetails/SMessagingDispatchStateTableRow.h
--- a/Engine/Plugins/Messaging/MessagingDebugger/Source/MessagingDebugger/Private/Widgets/MessageDetails/SMessagingDispatchStateTableRow.h
+++ b/Engine/Plugins/Messaging/MessagingDebugger/Source/MessagingDebugger/Private/Widgets/MessageDetails/SMessagingDispatchStateTableRow.h
@@ -3,6 +3,9 @@
 #pragma once
 
 
+class FMessagingDebuggerModel;
+
+
 #define LOCTEXT_NAMESPACE "SMessagingDispatchStateTableRow"
 
 
@@ -10,12 +13,12 @@
  * Implements a row widget for the dispatch state list.
  */
 class SMessagingDispatchStateTableRow
-	: public SMultiColumnTableRow<FMessageTracerDispatchStatePtr>
+	: public SMultiColumnTableRow<TSharedPtr<FMessageTracerDispatchState>>
 {
 public:
 
 	SLATE_BEGIN_ARGS(SMessagingDispatchStateTableRow) { }
-		SLATE_ARGUMENT(FMessageTracerDispatchStatePtr, DispatchState)
+		SLATE_ARGUMENT(TSharedPtr<FMessageTracerDispatchState>, DispatchState)
 		SLATE_ARGUMENT(TSharedPtr<ISlateStyle>, Style)
 	SLATE_END_ARGS()
 
@@ -27,7 +30,7 @@ public:
 	 * @param InArgs The construction arguments.
 	 * @param InOwnerTableView The table view that owns this row.
 	 */
-	void Construct( const FArguments& InArgs, const TSharedRef<STableViewBase>& InOwnerTableView, const FMessagingDebuggerModelRef& InModel )
+	void Construct(const FArguments& InArgs, const TSharedRef<STableViewBase>& InOwnerTableView, const TSharedRef<FMessagingDebuggerModel>& InModel)
 	{
 		check(InArgs._Style.IsValid());
 		check(InArgs._DispatchState.IsValid());
@@ -36,15 +39,15 @@ public:
 		Model = InModel;
 		Style = InArgs._Style;
 
-		SMultiColumnTableRow<FMessageTracerDispatchStatePtr>::Construct(FSuperRowType::FArguments(), InOwnerTableView);
+		SMultiColumnTableRow<TSharedPtr<FMessageTracerDispatchState>>::Construct(FSuperRowType::FArguments(), InOwnerTableView);
 	}
 
 public:
 
-	// SMultiColumnTableRow interface
+	//~ SMultiColumnTableRow interface
 
 	BEGIN_SLATE_FUNCTION_BUILD_OPTIMIZATION
-	virtual TSharedRef<SWidget> GenerateWidgetForColumn( const FName& ColumnName ) override
+	virtual TSharedRef<SWidget> GenerateWidgetForColumn(const FName& ColumnName) override
 	{
 		if (ColumnName == "DispatchLatency")
 		{
@@ -80,7 +83,7 @@ public:
 		}
 		else if (ColumnName == "Recipient")
 		{
-			FMessageTracerEndpointInfoPtr EndpointInfo = DispatchState->EndpointInfo;
+			TSharedPtr<FMessageTracerEndpointInfo> EndpointInfo = DispatchState->EndpointInfo;
 
 			return SNew(SHorizontalBox)
 
@@ -125,7 +128,7 @@ protected:
 	 * @param Latency The time span to convert.
 	 * @return The corresponding color value.
 	 */
-	FSlateColor TimespanToColor( double Timespan ) const
+	FSlateColor TimespanToColor(double Timespan) const
 	{
 		if (Timespan >= 0.01)
 		{
@@ -151,7 +154,7 @@ protected:
 	 * @param NamedThread The named thread.
 	 * @return The text representation.
 	 */
-	FText NamedThreadToReadableText( ENamedThreads::Type NamedThread )
+	FText NamedThreadToReadableText(ENamedThreads::Type NamedThread)
 	{
 		switch (NamedThread)
 		{
@@ -189,7 +192,7 @@ protected:
 	 * @return The text representation.
 	 * @todo gmp: refactor this into FText::AsTimespan or something like that
 	 */
-	FText TimespanToReadableText( double Seconds ) const
+	FText TimespanToReadableText(double Seconds) const
 	{
 		if (Seconds < 0.0)
 		{
@@ -296,10 +299,10 @@ private:
 private:
 
 	/** Holds the message dispatch state. */
-	FMessageTracerDispatchStatePtr DispatchState;
+	TSharedPtr<FMessageTracerDispatchState> DispatchState;
 
 	/** Holds a pointer to the view model. */
-	FMessagingDebuggerModelPtr Model;
+	TSharedPtr<FMessagingDebuggerModel> Model;
 
 	/** Holds the widget's visual style. */
 	TSharedPtr<ISlateStyle> Style;
diff --git a/Engine/Plugins/Messaging/MessagingDebugger/Source/MessagingDebugger/Private/Widgets/MessageDetails/SMessagingMessageDetails.cpp b/Engine/Plugins/Messaging/MessagingDebugger/Source/MessagingDebugger/Private/Widgets/MessageDetails/SMessagingMessageDetails.cpp
--- a/Engine/Plugins/Messaging/MessagingDebugger/Source/MessagingDebugger/Private/Widgets/MessageDetails/SMessagingMessageDetails.cpp
+++ b/Engine/Plugins/Messaging/MessagingDebugger/Source/MessagingDebugger/Private/Widgets/MessageDetails/SMessagingMessageDetails.cpp
@@ -1,6 +1,9 @@
 // Copyright 1998-2016 Epic Games, Inc. All Rights Reserved.
 
-#include "MessagingDebuggerPrivatePCH.h"
+#include "MessagingDebuggerPCH.h"
+#include "MessagingDebuggerModel.h"
+#include "SMessagingDispatchStateTableRow.h"
+#include "SMessagingMessageDetails.h"
 
 
 #define LOCTEXT_NAMESPACE "SMessagingMessageDetails"
@@ -21,7 +24,7 @@ SMessagingMessageDetails::~SMessagingMessageDetails()
 /* SMessagingMessageDetails interface
  *****************************************************************************/
 
-void SMessagingMessageDetails::Construct( const FArguments& InArgs, const FMessagingDebuggerModelRef& InModel, const TSharedRef<ISlateStyle>& InStyle )
+void SMessagingMessageDetails::Construct(const FArguments& InArgs, const TSharedRef<FMessagingDebuggerModel>& InModel, const TSharedRef<ISlateStyle>& InStyle)
 {
 	Model = InModel;
 	Style = InStyle;
@@ -37,7 +40,6 @@ void SMessagingMessageDetails::Construct( const FArguments& InArgs, const FMessa
 				SNew(SGridPanel)
 					.FillColumn(1, 1.0f)
 
-
 				// Sender thread
 				+ SGridPanel::Slot(0, 2)
 					.Padding(0.0f, 4.0f, 32.0f, 4.0f)
@@ -93,7 +95,7 @@ void SMessagingMessageDetails::Construct( const FArguments& InArgs, const FMessa
 					.Padding(0.0f)
 					[
 						// message list
-						SAssignNew(DispatchStateListView, SListView<FMessageTracerDispatchStatePtr>)
+						SAssignNew(DispatchStateListView, SListView<TSharedPtr<FMessageTracerDispatchState>>)
 							.ItemHeight(24.0f)
 							.ListItemsSource(&DispatchStateList)
 							.SelectionMode(ESelectionMode::None)
@@ -139,7 +141,7 @@ void SMessagingMessageDetails::Construct( const FArguments& InArgs, const FMessa
 
 void SMessagingMessageDetails::RefreshDetails()
 {
-	FMessageTracerMessageInfoPtr SelectedMessage = Model->GetSelectedMessage();
+	TSharedPtr<FMessageTracerMessageInfo> SelectedMessage = Model->GetSelectedMessage();
 
 	if (SelectedMessage.IsValid())
 	{
@@ -157,7 +159,7 @@ void SMessagingMessageDetails::RefreshDetails()
 /* SMessagingMessageDetails event handlers
  *****************************************************************************/
 
-TSharedRef<ITableRow> SMessagingMessageDetails::HandleDispatchStateListGenerateRow( FMessageTracerDispatchStatePtr DispatchState, const TSharedRef<STableViewBase>& OwnerTable )
+TSharedRef<ITableRow> SMessagingMessageDetails::HandleDispatchStateListGenerateRow(TSharedPtr<FMessageTracerDispatchState> DispatchState, const TSharedRef<STableViewBase>& OwnerTable)
 {
 	return SNew(SMessagingDispatchStateTableRow, OwnerTable, Model.ToSharedRef())
 		.DispatchState(DispatchState)
@@ -167,7 +169,7 @@ TSharedRef<ITableRow> SMessagingMessageDetails::HandleDispatchStateListGenerateR
 
 FText SMessagingMessageDetails::HandleExpirationText() const
 {
-	FMessageTracerMessageInfoPtr SelectedMessage = Model->GetSelectedMessage();
+	TSharedPtr<FMessageTracerMessageInfo> SelectedMessage = Model->GetSelectedMessage();
 
 	if (SelectedMessage.IsValid() && SelectedMessage->Context.IsValid())
 	{
@@ -193,7 +195,7 @@ void SMessagingMessageDetails::HandleModelSelectedMessageChanged()
 
 FText SMessagingMessageDetails::HandleSenderThreadText() const
 {
-	FMessageTracerMessageInfoPtr SelectedMessage = Model->GetSelectedMessage();
+	TSharedPtr<FMessageTracerMessageInfo> SelectedMessage = Model->GetSelectedMessage();
 
 	if (SelectedMessage.IsValid() && SelectedMessage->Context.IsValid())
 	{
@@ -246,7 +248,7 @@ FText SMessagingMessageDetails::HandleSenderThreadText() const
 
 FText SMessagingMessageDetails::HandleTimestampText() const
 {
-	FMessageTracerMessageInfoPtr SelectedMessage = Model->GetSelectedMessage();
+	TSharedPtr<FMessageTracerMessageInfo> SelectedMessage = Model->GetSelectedMessage();
 
 	if (SelectedMessage.IsValid() && SelectedMessage->Context.IsValid())
 	{
diff --git a/Engine/Plugins/Messaging/MessagingDebugger/Source/MessagingDebugger/Private/Widgets/MessageDetails/SMessagingMessageDetails.h b/Engine/Plugins/Messaging/MessagingDebugger/Source/MessagingDebugger/Private/Widgets/MessageDetails/SMessagingMessageDetails.h
--- a/Engine/Plugins/Messaging/MessagingDebugger/Source/MessagingDebugger/Private/Widgets/MessageDetails/SMessagingMessageDetails.h
+++ b/Engine/Plugins/Messaging/MessagingDebugger/Source/MessagingDebugger/Private/Widgets/MessageDetails/SMessagingMessageDetails.h
@@ -3,6 +3,9 @@
 #pragma once
 
 
+class FMessagingDebuggerModel;
+
+
 /**
  * Implements the message details panel.
  */
@@ -28,7 +31,7 @@ public:
 	 * @param InModel The view model to use.
 	 * @param InStyle The visual style to use for this widget.
 	 */
-	void Construct( const FArguments& InArgs, const FMessagingDebuggerModelRef& InModel, const TSharedRef<ISlateStyle>& InStyle );
+	void Construct(const FArguments& InArgs, const TSharedRef<FMessagingDebuggerModel>& InModel, const TSharedRef<ISlateStyle>& InStyle);
 
 protected:
 
@@ -38,7 +41,7 @@ protected:
 private:
 
 	/** Callback for generating a row widget for the dispatch state list view. */
-	TSharedRef<ITableRow> HandleDispatchStateListGenerateRow( FMessageTracerDispatchStatePtr DispatchState, const TSharedRef<STableViewBase>& OwnerTable );
+	TSharedRef<ITableRow> HandleDispatchStateListGenerateRow(TSharedPtr<FMessageTracerDispatchState> DispatchState, const TSharedRef<STableViewBase>& OwnerTable);
 
 	/** Callback for getting the text of the 'Expiration' field. */
 	FText HandleExpirationText() const;
@@ -55,13 +58,13 @@ private:
 private:
 
 	/** Holds the list of dispatch states. */
-	TArray<FMessageTracerDispatchStatePtr> DispatchStateList;
+	TArray<TSharedPtr<FMessageTracerDispatchState>> DispatchStateList;
 
 	/** Holds the dispatch state list view. */
-	TSharedPtr<SListView<FMessageTracerDispatchStatePtr> > DispatchStateListView;
+	TSharedPtr<SListView<TSharedPtr<FMessageTracerDispatchState>> > DispatchStateListView;
 
 	/** Holds a pointer to the view model. */
-	FMessagingDebuggerModelPtr Model;
+	TSharedPtr<FMessagingDebuggerModel> Model;
 
 	/** Holds the widget's visual style. */
 	TSharedPtr<ISlateStyle> Style;
diff --git a/Engine/Plugins/Messaging/MessagingDebugger/Source/MessagingDebugger/Private/Widgets/SMessagingDebugger.cpp b/Engine/Plugins/Messaging/MessagingDebugger/Source/MessagingDebugger/Private/Widgets/SMessagingDebugger.cpp
--- a/Engine/Plugins/Messaging/MessagingDebugger/Source/MessagingDebugger/Private/Widgets/SMessagingDebugger.cpp
+++ b/Engine/Plugins/Messaging/MessagingDebugger/Source/MessagingDebugger/Private/Widgets/SMessagingDebugger.cpp
@@ -1,7 +1,19 @@
 // Copyright 1998-2016 Epic Games, Inc. All Rights Reserved.
 
-#include "MessagingDebuggerPrivatePCH.h"
-#include "SDockTab.h"
+#include "MessagingDebuggerPCH.h"
+#include "MessagingDebuggerCommands.h"
+#include "MessagingDebuggerModel.h"
+#include "SMessagingBreakpoints.h"
+#include "SMessagingDebugger.h"
+#include "SMessagingEndpointDetails.h"
+#include "SMessagingEndpoints.h"
+#include "SMessagingGraph.h"
+#include "SMessagingInterceptors.h"
+#include "SMessagingMessageData.h"
+#include "SMessagingMessageDetails.h"
+#include "SMessagingHistory.h"
+#include "SMessagingTypes.h"
+#include "SMessagingDebuggerToolbar.h"
 
 
 #define LOCTEXT_NAMESPACE "SMessagingDebugger"
@@ -35,7 +47,13 @@ SMessagingDebugger::SMessagingDebugger()
 /* SMessagingDebugger interface
  *****************************************************************************/
 
-void SMessagingDebugger::Construct(const FArguments& InArgs, const TSharedRef<SDockTab>& ConstructUnderMajorTab, const TSharedPtr<SWindow>& ConstructUnderWindow, const IMessageTracerRef& InMessageTracer, const TSharedRef<ISlateStyle>& InStyle)
+void SMessagingDebugger::Construct(
+	const FArguments& InArgs,
+	const TSharedRef<SDockTab>& ConstructUnderMajorTab,
+	const TSharedPtr<SWindow>& ConstructUnderWindow,
+	const TSharedRef<IMessageTracer, ESPMode::ThreadSafe>& InMessageTracer,
+	const TSharedRef<ISlateStyle>& InStyle
+)
 {
 	MessageTracer = InMessageTracer;
 	Style = InStyle;
diff --git a/Engine/Plugins/Messaging/MessagingDebugger/Source/MessagingDebugger/Private/Widgets/SMessagingDebugger.h b/Engine/Plugins/Messaging/MessagingDebugger/Source/MessagingDebugger/Private/Widgets/SMessagingDebugger.h
--- a/Engine/Plugins/Messaging/MessagingDebugger/Source/MessagingDebugger/Private/Widgets/SMessagingDebugger.h
+++ b/Engine/Plugins/Messaging/MessagingDebugger/Source/MessagingDebugger/Private/Widgets/SMessagingDebugger.h
@@ -3,6 +3,10 @@
 #pragma once
 
 
+class FMessagingDebuggerModel;
+class IMessageTracer;
+
+
 /**
  * Implements the messaging system debugger application.
  */
@@ -30,7 +34,12 @@ public:
 	 * @param InMessageTracer The message tracer to use for the debugger.
 	 * @param InStyleSet The style set to use.
 	 */
-	void Construct(const FArguments& InArgs, const TSharedRef<SDockTab>& ConstructUnderMajorTab, const TSharedPtr<SWindow>& ConstructUnderWindow, const IMessageTracerRef& InMessageTracer, const TSharedRef<ISlateStyle>& InStyle);
+	void Construct(
+		const FArguments& InArgs,
+		const TSharedRef<SDockTab>& ConstructUnderMajorTab,
+		const TSharedPtr<SWindow>& ConstructUnderWindow,
+		const TSharedRef<IMessageTracer, ESPMode::ThreadSafe>& InMessageTracer,
+		const TSharedRef<ISlateStyle>& InStyle);
 
 protected:
 
@@ -105,10 +114,10 @@ private:
 	TSharedRef<FUICommandList> CommandList;
 
 	/** Holds a pointer to the message bus tracer. */
-	IMessageTracerPtr MessageTracer;
+	TSharedPtr<IMessageTracer, ESPMode::ThreadSafe> MessageTracer;
 
 	/** Holds the messaging debugger's view model. */
-	FMessagingDebuggerModelRef Model;
+	TSharedRef<FMessagingDebuggerModel> Model;
 
 	/** Holds the widget's visual style. */
 	TSharedPtr<ISlateStyle> Style;
diff --git a/Engine/Plugins/Messaging/MessagingDebugger/Source/MessagingDebugger/Private/Widgets/Toolbar/SMessagingDebuggerToolbar.cpp b/Engine/Plugins/Messaging/MessagingDebugger/Source/MessagingDebugger/Private/Widgets/Toolbar/SMessagingDebuggerToolbar.cpp
--- a/Engine/Plugins/Messaging/MessagingDebugger/Source/MessagingDebugger/Private/Widgets/Toolbar/SMessagingDebuggerToolbar.cpp
+++ b/Engine/Plugins/Messaging/MessagingDebugger/Source/MessagingDebugger/Private/Widgets/Toolbar/SMessagingDebuggerToolbar.cpp
@@ -1,6 +1,8 @@
 // Copyright 1998-2016 Epic Games, Inc. All Rights Reserved.
 
-#include "MessagingDebuggerPrivatePCH.h"
+#include "MessagingDebuggerPCH.h"
+#include "MessagingDebuggerCommands.h"
+#include "SMessagingDebuggerToolbar.h"
 
 
 #define LOCTEXT_NAMESPACE "SMessagingDebuggerToolbar"
@@ -9,7 +11,7 @@
 /* SMessagingDebuggerToolbar interface
  *****************************************************************************/
 
-void SMessagingDebuggerToolbar::Construct( const FArguments& InArgs, const TSharedRef<ISlateStyle>& InStyle, const TSharedRef<FUICommandList>& InCommandList )
+void SMessagingDebuggerToolbar::Construct(const FArguments& InArgs, const TSharedRef<ISlateStyle>& InStyle, const TSharedRef<FUICommandList>& InCommandList)
 {
 	ChildSlot
 	[
@@ -21,7 +23,7 @@ void SMessagingDebuggerToolbar::Construct( const FArguments& InArgs, const TShar
 /* SMessagingDebuggerToolbar implementation
  *****************************************************************************/
 
-TSharedRef<SWidget> SMessagingDebuggerToolbar::MakeToolbar( const TSharedRef<FUICommandList>& CommandList )
+TSharedRef<SWidget> SMessagingDebuggerToolbar::MakeToolbar(const TSharedRef<FUICommandList>& CommandList)
 {
 	FToolBarBuilder ToolBarBuilder(CommandList, FMultiBoxCustomization::None);
 
diff --git a/Engine/Plugins/Messaging/MessagingDebugger/Source/MessagingDebugger/Private/Widgets/Types/SMessagingTypes.cpp b/Engine/Plugins/Messaging/MessagingDebugger/Source/MessagingDebugger/Private/Widgets/Types/SMessagingTypes.cpp
--- a/Engine/Plugins/Messaging/MessagingDebugger/Source/MessagingDebugger/Private/Widgets/Types/SMessagingTypes.cpp
+++ b/Engine/Plugins/Messaging/MessagingDebugger/Source/MessagingDebugger/Private/Widgets/Types/SMessagingTypes.cpp
@@ -1,7 +1,11 @@
 // Copyright 1998-2016 Epic Games, Inc. All Rights Reserved.
 
-#include "MessagingDebuggerPrivatePCH.h"
-#include "SExpandableArea.h"
+#include "MessagingDebuggerPCH.h"
+#include "MessagingDebuggerModel.h"
+#include "MessagingDebuggerTypeFilter.h"
+#include "SMessagingTypes.h"
+#include "SMessagingTypesFilterBar.h"
+#include "SMessagingTypesTableRow.h"
 
 
 #define LOCTEXT_NAMESPACE "SMessagingTypes"
@@ -28,7 +32,7 @@ SMessagingTypes::~SMessagingTypes()
 /* SMessagingTypes interface
  *****************************************************************************/
 
-void SMessagingTypes::Construct( const FArguments& InArgs, const FMessagingDebuggerModelRef& InModel, const TSharedRef<ISlateStyle>& InStyle, const IMessageTracerRef& InTracer )
+void SMessagingTypes::Construct(const FArguments& InArgs, const TSharedRef<FMessagingDebuggerModel>& InModel, const TSharedRef<ISlateStyle>& InStyle, const TSharedRef<IMessageTracer, ESPMode::ThreadSafe>& InTracer)
 {
 	Filter = MakeShareable(new FMessagingDebuggerTypeFilter());
 	Model = InModel;
@@ -62,7 +66,7 @@ void SMessagingTypes::Construct( const FArguments& InArgs, const FMessagingDebug
 					.Padding(0.0f)
 					[
 						// type list
-						SAssignNew(TypeListView, SListView<FMessageTracerTypeInfoPtr>)
+						SAssignNew(TypeListView, SListView<TSharedPtr<FMessageTracerTypeInfo>>)
 							.ItemHeight(24.0f)
 							.ListItemsSource(&TypeList)
 							.SelectionMode(ESelectionMode::Multi)
@@ -127,7 +131,7 @@ void SMessagingTypes::Construct( const FArguments& InArgs, const FMessagingDebug
 /* SMessagingTypes implementation
  *****************************************************************************/
 
-void SMessagingTypes::AddType( const FMessageTracerTypeInfoRef& TypeInfo )
+void SMessagingTypes::AddType(const TSharedRef<FMessageTracerTypeInfo>& TypeInfo)
 {
 	if (Filter->FilterType(TypeInfo))
 	{
@@ -141,11 +145,11 @@ void SMessagingTypes::ReloadTypes()
 {
 	TypeList.Reset();
 	
-	TArray<FMessageTracerTypeInfoPtr> Types;
+	TArray<TSharedPtr<FMessageTracerTypeInfo>> Types;
 
 	if (Tracer->GetMessageTypes(Types) > 0)
 	{
-		for (TArray<FMessageTracerTypeInfoPtr>::TConstIterator It(Types); It; ++It)
+		for (TArray<TSharedPtr<FMessageTracerTypeInfo>>::TConstIterator It(Types); It; ++It)
 		{
 			AddType(It->ToSharedRef());
 		}
@@ -166,7 +170,7 @@ void SMessagingTypes::HandleFilterChanged()
 
 void SMessagingTypes::HandleModelSelectedMessageChanged()
 {
-	FMessageTracerMessageInfoPtr SelectedMessage = Model->GetSelectedMessage();
+	TSharedPtr<FMessageTracerMessageInfo> SelectedMessage = Model->GetSelectedMessage();
 
 	if (SelectedMessage.IsValid())
 	{
@@ -181,13 +185,13 @@ void SMessagingTypes::HandleTracerMessagesReset()
 }
 
 
-void SMessagingTypes::HandleTracerTypeAdded( FMessageTracerTypeInfoRef TypeInfo )
+void SMessagingTypes::HandleTracerTypeAdded(TSharedRef<FMessageTracerTypeInfo> TypeInfo)
 {
 	AddType(TypeInfo);
 }
 
 
-TSharedRef<ITableRow> SMessagingTypes::HandleTypeListGenerateRow( FMessageTracerTypeInfoPtr TypeInfo, const TSharedRef<STableViewBase>& OwnerTable )
+TSharedRef<ITableRow> SMessagingTypes::HandleTypeListGenerateRow(TSharedPtr<FMessageTracerTypeInfo> TypeInfo, const TSharedRef<STableViewBase>& OwnerTable)
 {
 	return SNew(SMessagingTypesTableRow, OwnerTable, Model.ToSharedRef())
 		.HighlightText(this, &SMessagingTypes::HandleTypeListGetHighlightText)
@@ -203,7 +207,7 @@ FText SMessagingTypes::HandleTypeListGetHighlightText() const
 }
 
 
-void SMessagingTypes::HandleTypeListSelectionChanged( FMessageTracerTypeInfoPtr InItem, ESelectInfo::Type SelectInfo )
+void SMessagingTypes::HandleTypeListSelectionChanged(TSharedPtr<FMessageTracerTypeInfo> InItem, ESelectInfo::Type SelectInfo)
 {
 
 }
diff --git a/Engine/Plugins/Messaging/MessagingDebugger/Source/MessagingDebugger/Private/Widgets/Types/SMessagingTypes.h b/Engine/Plugins/Messaging/MessagingDebugger/Source/MessagingDebugger/Private/Widgets/Types/SMessagingTypes.h
--- a/Engine/Plugins/Messaging/MessagingDebugger/Source/MessagingDebugger/Private/Widgets/Types/SMessagingTypes.h
+++ b/Engine/Plugins/Messaging/MessagingDebugger/Source/MessagingDebugger/Private/Widgets/Types/SMessagingTypes.h
@@ -3,6 +3,11 @@
 #pragma once
 
 
+class FMessagingDebuggerModel;
+class FMessagingDebuggerTypeFilter;
+class IMessageTracer;
+
+
 /**
  * Implements the message types panel.
  */
@@ -29,7 +34,11 @@ public:
 	 * @param InStyle The visual style to use for this widget.
 	 * @param InTracer The message tracer.
 	 */
-	void Construct( const FArguments& InArgs, const FMessagingDebuggerModelRef& InModel, const TSharedRef<ISlateStyle>& InStyle, const IMessageTracerRef& InTracer );
+	void Construct(
+		const FArguments& InArgs,
+		const TSharedRef<FMessagingDebuggerModel>& InModel,
+		const TSharedRef<ISlateStyle>& InStyle,
+		const TSharedRef<IMessageTracer, ESPMode::ThreadSafe>& InTracer);
 
 protected:
 
@@ -38,7 +47,7 @@ protected:
 	 *
 	 * @param TypeInfo The information about the message type to add.
 	 */
-	void AddType( const FMessageTracerTypeInfoRef& TypeInfo );
+	void AddType(const TSharedRef<FMessageTracerTypeInfo>& TypeInfo);
 
 	/** Reloads the message history. */
 	void ReloadTypes();
@@ -55,34 +64,34 @@ private:
 	void HandleTracerMessagesReset();
 
 	/** Callback for when a message type has been added to the message tracer. */
-	void HandleTracerTypeAdded( FMessageTracerTypeInfoRef TypeInfo );
+	void HandleTracerTypeAdded(TSharedRef<FMessageTracerTypeInfo> TypeInfo);
 
 	/** Callback for generating a row widget for the message type list view. */
-	TSharedRef<ITableRow> HandleTypeListGenerateRow( FMessageTracerTypeInfoPtr TypeInfo, const TSharedRef<STableViewBase>& OwnerTable );
+	TSharedRef<ITableRow> HandleTypeListGenerateRow(TSharedPtr<FMessageTracerTypeInfo> TypeInfo, const TSharedRef<STableViewBase>& OwnerTable);
 
 	/** Callback for getting the highlight string for message types. */
 	FText HandleTypeListGetHighlightText() const;
 
 	/** Callback for selecting message types. */
-	void HandleTypeListSelectionChanged( FMessageTracerTypeInfoPtr InItem, ESelectInfo::Type SelectInfo );
+	void HandleTypeListSelectionChanged(TSharedPtr<FMessageTracerTypeInfo> InItem, ESelectInfo::Type SelectInfo);
 
 private:
 
 	/** Holds the message type filter model. */
-	FMessagingDebuggerTypeFilterPtr Filter;
+	TSharedPtr<FMessagingDebuggerTypeFilter> Filter;
 
 	/** Holds a pointer to the view model. */
-	FMessagingDebuggerModelPtr Model;
+	TSharedPtr<FMessagingDebuggerModel> Model;
 
 	/** Holds the widget's visual style. */
 	TSharedPtr<ISlateStyle> Style;
 
 	/** Holds a pointer to the message bus tracer. */
-	IMessageTracerPtr Tracer;
+	TSharedPtr<IMessageTracer, ESPMode::ThreadSafe> Tracer;
 
 	/** Holds the filtered list of message types. */
-	TArray<FMessageTracerTypeInfoPtr> TypeList;
+	TArray<TSharedPtr<FMessageTracerTypeInfo>> TypeList;
 
 	/** Holds the message type list view. */
-	TSharedPtr<SListView<FMessageTracerTypeInfoPtr>> TypeListView;
+	TSharedPtr<SListView<TSharedPtr<FMessageTracerTypeInfo>>> TypeListView;
 };
diff --git a/Engine/Plugins/Messaging/MessagingDebugger/Source/MessagingDebugger/Private/Widgets/Types/SMessagingTypesFilterBar.cpp b/Engine/Plugins/Messaging/MessagingDebugger/Source/MessagingDebugger/Private/Widgets/Types/SMessagingTypesFilterBar.cpp
--- a/Engine/Plugins/Messaging/MessagingDebugger/Source/MessagingDebugger/Private/Widgets/Types/SMessagingTypesFilterBar.cpp
+++ b/Engine/Plugins/Messaging/MessagingDebugger/Source/MessagingDebugger/Private/Widgets/Types/SMessagingTypesFilterBar.cpp
@@ -1,8 +1,8 @@
 // Copyright 1998-2016 Epic Games, Inc. All Rights Reserved.
 
-#include "MessagingDebuggerPrivatePCH.h"
+#include "MessagingDebuggerPCH.h"
+#include "MessagingDebuggerTypeFilter.h"
 #include "SMessagingTypesFilterBar.h"
-#include "SSearchBox.h"
 
 
 #define LOCTEXT_NAMESPACE "SMessagingTypesFilterBar"
@@ -11,24 +11,26 @@
 /* SMessagingTypesFilterBar interface
 *****************************************************************************/
 
-void SMessagingTypesFilterBar::Construct(const FArguments& InArgs, FMessagingDebuggerTypeFilterRef InFilter)
+void SMessagingTypesFilterBar::Construct(const FArguments& InArgs, TSharedRef<FMessagingDebuggerTypeFilter> InFilter)
 {
 	Filter = InFilter;
 
 	ChildSlot
-		[
-			SNew(SHorizontalBox)
+	[
+		SNew(SHorizontalBox)
 
-			+ SHorizontalBox::Slot()
+		+ SHorizontalBox::Slot()
 			.FillWidth(1.0f)
 			.VAlign(VAlign_Top)
 			[
 				// search box
 				SNew(SSearchBox)
-				.HintText(LOCTEXT("SearchBoxHint", "Search message types"))
-				.OnTextChanged(this, &SMessagingTypesFilterBar::HandleFilterStringTextChanged)
+					.HintText(LOCTEXT("SearchBoxHint", "Search message types"))
+					.OnTextChanged_Lambda([this](const FText& NewText) {
+						Filter->SetFilterString(NewText.ToString());
+					})
 			]
-		];
+	];
 }
 
 #undef LOCTEXT_NAMESPACE
diff --git a/Engine/Plugins/Messaging/MessagingDebugger/Source/MessagingDebugger/Private/Widgets/Types/SMessagingTypesFilterBar.h b/Engine/Plugins/Messaging/MessagingDebugger/Source/MessagingDebugger/Private/Widgets/Types/SMessagingTypesFilterBar.h
--- a/Engine/Plugins/Messaging/MessagingDebugger/Source/MessagingDebugger/Private/Widgets/Types/SMessagingTypesFilterBar.h
+++ b/Engine/Plugins/Messaging/MessagingDebugger/Source/MessagingDebugger/Private/Widgets/Types/SMessagingTypesFilterBar.h
@@ -3,6 +3,9 @@
 #pragma once
 
 
+class FMessagingDebuggerTypeFilter;
+
+
 /**
  * Implements the message type list filter bar widget.
  */
@@ -26,18 +29,10 @@ public:
 	 * @param InArgs The declaration data for this widget.
 	 * @param InFilter The filter model.
 	 */
-	void Construct( const FArguments& InArgs, FMessagingDebuggerTypeFilterRef InFilter );
-
-private:
-
-	/** Handles changing the filter string text box text. */
-	void HandleFilterStringTextChanged( const FText& NewText )
-	{
-		Filter->SetFilterString(NewText.ToString());
-	}
+	void Construct( const FArguments& InArgs, TSharedRef<FMessagingDebuggerTypeFilter> InFilter );
 
 private:
 
 	/** Holds the filter model. */
-	FMessagingDebuggerTypeFilterPtr Filter;
+	TSharedPtr<FMessagingDebuggerTypeFilter> Filter;
 };
diff --git a/Engine/Plugins/Messaging/MessagingDebugger/Source/MessagingDebugger/Private/Widgets/Types/SMessagingTypesTableRow.h b/Engine/Plugins/Messaging/MessagingDebugger/Source/MessagingDebugger/Private/Widgets/Types/SMessagingTypesTableRow.h
--- a/Engine/Plugins/Messaging/MessagingDebugger/Source/MessagingDebugger/Private/Widgets/Types/SMessagingTypesTableRow.h
+++ b/Engine/Plugins/Messaging/MessagingDebugger/Source/MessagingDebugger/Private/Widgets/Types/SMessagingTypesTableRow.h
@@ -3,20 +3,21 @@
 #pragma once
 
 
-#define LOCTEXT_NAMESPACE "SMessagingTypesTableRow"
+struct FMessageTracerTypeInfo;
+class FMessagingDebuggerModel;
 
 
 /**
  * Implements a row widget for the message type list.
  */
 class SMessagingTypesTableRow
-	: public SMultiColumnTableRow<FMessageTracerTypeInfoPtr>
+	: public SMultiColumnTableRow<TSharedPtr<FMessageTracerTypeInfo>>
 {
 public:
 
 	SLATE_BEGIN_ARGS(SMessagingTypesTableRow) { }
 		SLATE_ATTRIBUTE(FText, HighlightText)
-		SLATE_ARGUMENT(FMessageTracerTypeInfoPtr, TypeInfo)
+		SLATE_ARGUMENT(TSharedPtr<FMessageTracerTypeInfo>, TypeInfo)
 		SLATE_ARGUMENT(TSharedPtr<ISlateStyle>, Style)
 	SLATE_END_ARGS()
 
@@ -28,99 +29,13 @@ public:
 	 * @param InArgs The construction arguments.
 	 * @param InOwnerTableView The table view that owns this row.
 	 */
-	void Construct( const FArguments& InArgs, const TSharedRef<STableViewBase>& InOwnerTableView, const FMessagingDebuggerModelRef& InModel )
-	{
-		check(InArgs._Style.IsValid());
-		check(InArgs._TypeInfo.IsValid());
-
-		HighlightText = InArgs._HighlightText;
-		Model = InModel;
-		Style = InArgs._Style;
-		TypeInfo = InArgs._TypeInfo;
-
-		SMultiColumnTableRow<FMessageTracerTypeInfoPtr>::Construct(FSuperRowType::FArguments(), InOwnerTableView);
-	}
+	void Construct(const FArguments& InArgs, const TSharedRef<STableViewBase>& InOwnerTableView, const TSharedRef<FMessagingDebuggerModel>& InModel);
 
 public:
 
-	// SMultiColumnTableRow interface
-
-	BEGIN_SLATE_FUNCTION_BUILD_OPTIMIZATION
-	virtual TSharedRef<SWidget> GenerateWidgetForColumn( const FName& ColumnName ) override
-	{
-		if (ColumnName == "Break")
-		{
-			return SNew(SBorder)
-				.BorderBackgroundColor(FLinearColor(0.0f, 0.0f, 0.0f, 0.33f))
-				.BorderImage(Style->GetBrush("BreakpointBorder"));
-
-/*			return SNew(SImage)
-				.Image(Style->GetBrush("BreakDisabled"));*/
-		}
-		else if (ColumnName == "Messages")
-		{
-			return SNew(SBox)
-				.Padding(FMargin(4.0f, 0.0f))
-				.VAlign(VAlign_Center)
-				[
-					SNew(STextBlock)
-						.HighlightText(HighlightText)
-						.Text(this, &SMessagingTypesTableRow::HandleMessagesText)
-				];
-		}
-		else if (ColumnName == "Name")
-		{
-			return SNew(SBox)
-				.Padding(FMargin(4.0f, 0.0f))
-				.VAlign(VAlign_Center)
-				[
-					SNew(STextBlock)
-						.HighlightText(HighlightText)
-						.Text(FText::FromName(TypeInfo->TypeName))
-				];
-		}
-		else if (ColumnName == "Visibility")
-		{
-			return SNew(SBox)
-				.HAlign(HAlign_Center)
-				.VAlign(VAlign_Center)
-				[
-					SNew(SCheckBox)
-						.Style(&Style->GetWidgetStyle<FCheckBoxStyle>("VisibilityCheckbox"))
-						.IsChecked(this, &SMessagingTypesTableRow::HandleVisibilityCheckBoxIsChecked)
-						.OnCheckStateChanged(this, &SMessagingTypesTableRow::HandleVisibilityCheckBoxCheckStateChanged)
-						.ToolTipText(LOCTEXT("VisibilityCheckboxTooltipText", "Toggle visibility of messages of this type"))
-				];
-		}
-
-		return SNullWidget::NullWidget;
-	}
-	END_SLATE_FUNCTION_BUILD_OPTIMIZATION
-
-private:
+	//~ SMultiColumnTableRow interface
 
-	/** Gets the text for the Messages column. */
-	FText HandleMessagesText() const
-	{
-		return FText::AsNumber(TypeInfo->Messages.Num());
-	}
-
-	/** Handles changing the checked state of the visibility check box. */
-	void HandleVisibilityCheckBoxCheckStateChanged( ECheckBoxState CheckState )
-	{
-		Model->SetTypeVisibility(TypeInfo.ToSharedRef(), (CheckState == ECheckBoxState::Checked));
-	}
-
-	/** Gets the image for the visibility check box. */
-	ECheckBoxState HandleVisibilityCheckBoxIsChecked() const
-	{
-		if (Model->IsTypeVisible(TypeInfo.ToSharedRef()))
-		{
-			return ECheckBoxState::Checked;
-		}
-
-		return ECheckBoxState::Unchecked;
-	}
+	virtual TSharedRef<SWidget> GenerateWidgetForColumn(const FName& ColumnName) override;
 
 private:
 
@@ -128,14 +43,11 @@ private:
 	TAttribute<FText> HighlightText;
 
 	/** Holds a pointer to the view model. */
-	FMessagingDebuggerModelPtr Model;
+	TSharedPtr<FMessagingDebuggerModel> Model;
 
 	/** Holds the widget's visual style. */
 	TSharedPtr<ISlateStyle> Style;
 
 	/** Holds message type's debug information. */
-	FMessageTracerTypeInfoPtr TypeInfo;
+	TSharedPtr<FMessageTracerTypeInfo> TypeInfo;
 };
-
-
-#undef LOCTEXT_NAMESPACE
diff --git a/Engine/Plugins/Messaging/TcpMessaging/Source/TcpMessaging/Private/TcpMessagingModule.cpp b/Engine/Plugins/Messaging/TcpMessaging/Source/TcpMessaging/Private/TcpMessagingModule.cpp
--- a/Engine/Plugins/Messaging/TcpMessaging/Source/TcpMessaging/Private/TcpMessagingModule.cpp
+++ b/Engine/Plugins/Messaging/TcpMessaging/Source/TcpMessaging/Private/TcpMessagingModule.cpp
@@ -6,7 +6,7 @@
 #include "ISettingsSection.h"
 #endif
 #include "ModuleManager.h"
-#include "IConnectionBasedMessagingModule.h"
+#include "ITcpMessagingModule.h"
 
 
 DEFINE_LOG_CATEGORY(LogTcpMessaging);
@@ -19,7 +19,7 @@ DEFINE_LOG_CATEGORY(LogTcpMessaging);
  */
 class FTcpMessagingModule
 	: public FSelfRegisteringExec
-	, public IConnectionBasedMessagingModule
+	, public ITcpMessagingModule
 {
 public:
 
@@ -311,7 +311,7 @@ private:
 private:
 
 	/** Holds the message bridge if present. */
-	IMessageBridgePtr MessageBridge;
+	TSharedPtr<IMessageBridge, ESPMode::ThreadSafe> MessageBridge;
 	
 	/** Message transport pointer, if valid */
 	TWeakPtr<FTcpMessageTransport, ESPMode::ThreadSafe> MessageTransportPtr;
diff --git a/Engine/Plugins/Messaging/TcpMessaging/Source/TcpMessaging/Private/Transport/TcpDeserializedMessage.cpp b/Engine/Plugins/Messaging/TcpMessaging/Source/TcpMessaging/Private/Transport/TcpDeserializedMessage.cpp
--- a/Engine/Plugins/Messaging/TcpMessaging/Source/TcpMessaging/Private/Transport/TcpDeserializedMessage.cpp
+++ b/Engine/Plugins/Messaging/TcpMessaging/Source/TcpMessaging/Private/Transport/TcpDeserializedMessage.cpp
@@ -8,11 +8,10 @@
 /* FTcpDeserializedMessage structors
 *****************************************************************************/
 
-FTcpDeserializedMessage::FTcpDeserializedMessage(const IMessageAttachmentPtr& InAttachment)
+FTcpDeserializedMessage::FTcpDeserializedMessage(const TSharedPtr<IMessageAttachment, ESPMode::ThreadSafe>& InAttachment)
 	: Attachment(InAttachment)
 	, MessageData(nullptr)
-{
-}
+{ }
 
 
 FTcpDeserializedMessage::~FTcpDeserializedMessage()
@@ -136,7 +135,7 @@ const TMap<FName, FString>& FTcpDeserializedMessage::GetAnnotations() const
 }
 
 
-IMessageAttachmentPtr FTcpDeserializedMessage::GetAttachment() const
+TSharedPtr<IMessageAttachment, ESPMode::ThreadSafe> FTcpDeserializedMessage::GetAttachment() const
 {
 	return Attachment;
 }
diff --git a/Engine/Plugins/Messaging/TcpMessaging/Source/TcpMessaging/Private/Transport/TcpDeserializedMessage.h b/Engine/Plugins/Messaging/TcpMessaging/Source/TcpMessaging/Private/Transport/TcpDeserializedMessage.h
--- a/Engine/Plugins/Messaging/TcpMessaging/Source/TcpMessaging/Private/Transport/TcpDeserializedMessage.h
+++ b/Engine/Plugins/Messaging/TcpMessaging/Source/TcpMessaging/Private/Transport/TcpDeserializedMessage.h
@@ -16,7 +16,7 @@ public:
 	 *
 	 * @param InAttachment An optional message attachment.
 	 */
-	FTcpDeserializedMessage(const IMessageAttachmentPtr& InAttachment);
+	FTcpDeserializedMessage(const TSharedPtr<IMessageAttachment, ESPMode::ThreadSafe>& InAttachment);
 
 	/** Virtual destructor. */
 	virtual ~FTcpDeserializedMessage() override;
@@ -36,7 +36,7 @@ public:
 	// IMessageContext interface
 
 	virtual const TMap<FName, FString>& GetAnnotations() const override;
-	virtual IMessageAttachmentPtr GetAttachment() const override;
+	virtual TSharedPtr<IMessageAttachment, ESPMode::ThreadSafe> GetAttachment() const override;
 	virtual const FDateTime& GetExpiration() const override;
 	virtual const void* GetMessage() const override;
 	virtual const TWeakObjectPtr<UScriptStruct>& GetMessageTypeInfo() const override;
@@ -54,7 +54,7 @@ private:
 	TMap<FName, FString> Annotations;
 
 	/** Holds a pointer to attached binary data. */
-	IMessageAttachmentPtr Attachment;
+	TSharedPtr<IMessageAttachment, ESPMode::ThreadSafe> Attachment;
 
 	/** Holds the expiration time. */
 	FDateTime Expiration;
diff --git a/Engine/Plugins/Messaging/UdpMessaging/Source/UdpMessaging/Private/Shared/UdpMessageSegment.h b/Engine/Plugins/Messaging/UdpMessaging/Source/UdpMessaging/Private/Shared/UdpMessageSegment.h
--- a/Engine/Plugins/Messaging/UdpMessaging/Source/UdpMessaging/Private/Shared/UdpMessageSegment.h
+++ b/Engine/Plugins/Messaging/UdpMessaging/Source/UdpMessaging/Private/Shared/UdpMessageSegment.h
@@ -3,38 +3,35 @@
 #pragma once
 
 
-namespace EUdpMessageSegments
+/**
+ * Enumerates message segment types.
+ */
+enum class EUdpMessageSegments : uint8
 {
-	/**
-	 * Enumerates message segment types.
-	 */
-	enum Type
-	{
-		/** None. */
-		None,
+	/** None. */
+	None,
 
-		/** Request to abort the sending of a message. */
-		Abort,
+	/** Request to abort the sending of a message. */
+	Abort,
 
-		/** Acknowledges that the message was received successfully. */
-		Acknowledge,
+	/** Acknowledges that the message was received successfully. */
+	Acknowledge,
 
-		/** Notifies the bus that an endpoint has left. */
-		Bye,
+	/** Notifies the bus that an endpoint has left. */
+	Bye,
 
-		/** A message data segment. */
-		Data,
+	/** A message data segment. */
+	Data,
 
-		/** Notifies the bus that an endpoint has joined. */
-		Hello,
+	/** Notifies the bus that an endpoint has joined. */
+	Hello,
 
-		/** Request to retransmit selected data segments. */
-		Retransmit,
+	/** Request to retransmit selected data segments. */
+	Retransmit,
 
-		/** Notification that an inbound message timed out. */
-		Timeout
-	};
-}
+	/** Notification that an inbound message timed out. */
+	Timeout
+};
 
 
 namespace FUdpMessageSegment
@@ -58,7 +55,7 @@ namespace FUdpMessageSegment
 		 *
 		 * @see EUdpMessageSegments
 		 */
-		uint8 SegmentType;
+		EUdpMessageSegments SegmentType;
 
 	public:
 
diff --git a/Engine/Plugins/Messaging/UdpMessaging/Source/UdpMessaging/Private/Tests/UdpMessageSegmenterTest.cpp b/Engine/Plugins/Messaging/UdpMessaging/Source/UdpMessaging/Private/Tests/UdpMessageSegmenterTest.cpp
--- a/Engine/Plugins/Messaging/UdpMessaging/Source/UdpMessaging/Private/Tests/UdpMessageSegmenterTest.cpp
+++ b/Engine/Plugins/Messaging/UdpMessaging/Source/UdpMessaging/Private/Tests/UdpMessageSegmenterTest.cpp
@@ -1,7 +1,8 @@
 // Copyright 1998-2016 Epic Games, Inc. All Rights Reserved.
 
 #include "UdpMessagingPrivatePCH.h"
-#include "AutomationTest.h"
+#include "UdpMessageSegmenter.h"
+#include "UdpSerializedMessage.h"
 
 
 IMPLEMENT_SIMPLE_AUTOMATION_TEST(FUdpMessageSegmenterTest, "System.Core.Messaging.Transports.Udp.UdpMessageSegmenter", EAutomationTestFlags::ApplicationContextMask | EAutomationTestFlags::EngineFilter)
@@ -15,7 +16,7 @@ void RunSegmentationTest(FAutomationTestBase& Test, uint32 MessageSize, uint16 S
 	Test.AddLogItem(FString::Printf(TEXT("Segmenting message of size %i with %i segments of size %i..."), MessageSize, NumSegments, SegmentSize));
 
 	// create a large message to segment
-	FUdpSerializedMessageRef Message = MakeShareable(new FUdpSerializedMessage());
+	TSharedRef<FUdpSerializedMessage, ESPMode::ThreadSafe> Message = MakeShareable(new FUdpSerializedMessage());
 
 	for (uint8 SegmentIndex = 0; SegmentIndex < NumSegments; ++SegmentIndex)
 	{
@@ -123,7 +124,9 @@ bool FUdpMessageSegmenterTest::RunTest(const FString& Parameters)
 	return true;
 }
 
+
 void EmptyLinkFunctionForStaticInitializationUdpMessageSegmenterTest()
 {
-	// This function exists to prevent the object file containing this test from being excluded by the linker, because it has no publically referenced symbols.
+	// This function exists to prevent the object file containing this test from
+	// being excluded by the linker, because it has no publicly referenced symbols.
 }
diff --git a/Engine/Plugins/Messaging/UdpMessaging/Source/UdpMessaging/Private/Tests/UdpMessageTransportTest.cpp b/Engine/Plugins/Messaging/UdpMessaging/Source/UdpMessaging/Private/Tests/UdpMessageTransportTest.cpp
--- a/Engine/Plugins/Messaging/UdpMessaging/Source/UdpMessaging/Private/Tests/UdpMessageTransportTest.cpp
+++ b/Engine/Plugins/Messaging/UdpMessaging/Source/UdpMessaging/Private/Tests/UdpMessageTransportTest.cpp
@@ -1,7 +1,7 @@
 // Copyright 1998-2016 Epic Games, Inc. All Rights Reserved.
 
 #include "UdpMessagingPrivatePCH.h"
-#include "AutomationTest.h"
+#include "UdpMessageTransport.h"
 #include "UdpMessagingTestTypes.h"
 
 
@@ -41,7 +41,7 @@ public:
 		return NumReceivedMessages;
 	}
 
-	bool Publish(const IMessageContextRef& Context)
+	bool Publish(const TSharedRef<IMessageContext, ESPMode::ThreadSafe>& Context)
 	{
 		return Transport->TransportMessage(Context, TArray<FGuid>());
 	}
@@ -58,7 +58,7 @@ public:
 
 private:
 
-	void HandleTransportMessageReceived(const IMessageContextRef& MessageContext, const FGuid& NodeId)
+	void HandleTransportMessageReceived(const TSharedRef<IMessageContext, ESPMode::ThreadSafe>& MessageContext, const FGuid& NodeId)
 	{
 		FPlatformAtomics::InterlockedIncrement(&NumReceivedMessages);
 	}
@@ -79,7 +79,7 @@ private:
 	TArray<FGuid> LostNodes;
 	int32 NumReceivedMessages;
 	FAutomationTestBase& Test;
-	IMessageTransportPtr Transport;
+	TSharedPtr<IMessageTransport, ESPMode::ThreadSafe> Transport;
 };
 
 
@@ -137,7 +137,7 @@ bool FUdpMessageTransportTest::RunTest(const FString& Parameters)
 	for (int32 Count = 0; Count < NumTestMessages; ++Count)
 	{
 		FUdpMockMessage* Message = new FUdpMockMessage(MessageSize);
-		IMessageContextRef Context = MakeShareable(new FUdpMockMessageContext(Message));
+		TSharedRef<IMessageContext, ESPMode::ThreadSafe> Context = MakeShareable(new FUdpMockMessageContext(Message));
 
 		Transport1.Publish(Context);
 	}
diff --git a/Engine/Plugins/Messaging/UdpMessaging/Source/UdpMessaging/Private/Tests/UdpMessagingTestTypes.h b/Engine/Plugins/Messaging/UdpMessaging/Source/UdpMessaging/Private/Tests/UdpMessagingTestTypes.h
--- a/Engine/Plugins/Messaging/UdpMessaging/Source/UdpMessaging/Private/Tests/UdpMessagingTestTypes.h
+++ b/Engine/Plugins/Messaging/UdpMessaging/Source/UdpMessaging/Private/Tests/UdpMessagingTestTypes.h
@@ -74,7 +74,7 @@ public:
 private:
 
 	TMap<FName, FString> Annotations;
-	IMessageAttachmentPtr Attachment;
+	TSharedPtr<IMessageAttachment, ESPMode::ThreadSafe> Attachment;
 	FDateTime Expiration;
 	void* Message;
 	IMessageContextPtr OriginalContext;
diff --git a/Engine/Plugins/Messaging/UdpMessaging/Source/UdpMessaging/Private/Tests/UdpSerializeMessageTaskTest.cpp b/Engine/Plugins/Messaging/UdpMessaging/Source/UdpMessaging/Private/Tests/UdpSerializeMessageTaskTest.cpp
--- a/Engine/Plugins/Messaging/UdpMessaging/Source/UdpMessaging/Private/Tests/UdpSerializeMessageTaskTest.cpp
+++ b/Engine/Plugins/Messaging/UdpMessaging/Source/UdpMessaging/Private/Tests/UdpSerializeMessageTaskTest.cpp
@@ -1,9 +1,9 @@
 // Copyright 1998-2016 Epic Games, Inc. All Rights Reserved.
 
 #include "UdpMessagingPrivatePCH.h"
-#include "AutomationTest.h"
 #include "UdpMessagingTestTypes.h"
-
+#include "UdpSerializeMessageTask.h"
+#include "UdpSerializedMessage.h"
 
 
 IMPLEMENT_SIMPLE_AUTOMATION_TEST(FUdpSerializeMessageTaskTest, "System.Core.Messaging.Transports.Udp.UdpSerializedMessage", EAutomationTestFlags::ApplicationContextMask | EAutomationTestFlags::EngineFilter)
@@ -16,11 +16,11 @@ namespace UdpSerializeMessageTaskTest
 
 	int32 CompletedMessages = 0;
 	int32 FailedMessages = 0;
-	FUdpSerializedMessagePtr ReferenceMessage;
+	TSharedPtr<FUdpSerializedMessage, ESPMode::ThreadSafe> ReferenceMessage;
 }
 
 
-void HandleSerializedMessageStateChanged(FUdpSerializedMessageRef SerializedMessage)
+void HandleSerializedMessageStateChanged(TSharedRef<FUdpSerializedMessage, ESPMode::ThreadSafe> SerializedMessage)
 {
 	using namespace UdpSerializeMessageTaskTest;
 
@@ -54,7 +54,7 @@ bool FUdpSerializeMessageTaskTest::RunTest(const FString& Parameters)
 	FailedMessages = 0;
 	ReferenceMessage = MakeShareable(new FUdpSerializedMessage);
 
-	IMessageContextRef Context = MakeShareable(new FUdpMockMessageContext(new FUdpMockMessage));
+	TSharedRef<IMessageContext, ESPMode::ThreadSafe> Context = MakeShareable(new FUdpMockMessageContext(new FUdpMockMessage));
 
 	// synchronous reference serialization
 	FUdpSerializeMessageTask ReferenceTask(Context, ReferenceMessage.ToSharedRef());
@@ -65,7 +65,7 @@ bool FUdpSerializeMessageTaskTest::RunTest(const FString& Parameters)
 	// stress test
 	for (int32 TestIndex = 0; TestIndex < NumMessages; ++TestIndex)
 	{
-		FUdpSerializedMessageRef SerializedMessage = MakeShareable(new FUdpSerializedMessage);
+		TSharedRef<FUdpSerializedMessage, ESPMode::ThreadSafe> SerializedMessage = MakeShareable(new FUdpSerializedMessage);
 		{
 			SerializedMessage->OnStateChanged().BindStatic(&HandleSerializedMessageStateChanged, SerializedMessage);
 		}
diff --git a/Engine/Plugins/Messaging/UdpMessaging/Source/UdpMessaging/Private/Transport/UdpDeserializedMessage.cpp b/Engine/Plugins/Messaging/UdpMessaging/Source/UdpMessaging/Private/Transport/UdpDeserializedMessage.cpp
--- a/Engine/Plugins/Messaging/UdpMessaging/Source/UdpMessaging/Private/Transport/UdpDeserializedMessage.cpp
+++ b/Engine/Plugins/Messaging/UdpMessaging/Source/UdpMessaging/Private/Transport/UdpDeserializedMessage.cpp
@@ -1,18 +1,17 @@
 // Copyright 1998-2016 Epic Games, Inc. All Rights Reserved.
 
 #include "UdpMessagingPrivatePCH.h"
-#include "JsonStructDeserializerBackend.h"
-#include "StructDeserializer.h"
+#include "UdpDeserializedMessage.h"
+#include "UdpReassembledMessage.h"
 
 
 /* FUdpDeserializedMessage structors
 *****************************************************************************/
 
-FUdpDeserializedMessage::FUdpDeserializedMessage(const IMessageAttachmentPtr& InAttachment)
+FUdpDeserializedMessage::FUdpDeserializedMessage(const TSharedPtr<IMessageAttachment, ESPMode::ThreadSafe>& InAttachment)
 	: Attachment(InAttachment)
 	, MessageData(nullptr)
-{
-}
+{ }
 
 
 FUdpDeserializedMessage::~FUdpDeserializedMessage()
@@ -33,12 +32,12 @@ FUdpDeserializedMessage::~FUdpDeserializedMessage()
 /* FUdpDeserializedMessage interface
  *****************************************************************************/
 
-bool FUdpDeserializedMessage::Deserialize(const FUdpReassembledMessageRef& ReassembledMessage)
+bool FUdpDeserializedMessage::Deserialize(const FReassembledUdpMessage& ReassembledMessage)
 {
 	// Note that some complex values are deserialized manually here, so that we
 	// can sanity check their values. @see FUdpSerializeMessageTask::DoTask()
 
-	FMemoryReader MessageReader(ReassembledMessage->GetData());
+	FMemoryReader MessageReader(ReassembledMessage.GetData());
 	MessageReader.ArMaxSerializeSize = NAME_SIZE;
 
 	// message type info
@@ -136,7 +135,7 @@ const TMap<FName, FString>& FUdpDeserializedMessage::GetAnnotations() const
 }
 
 
-IMessageAttachmentPtr FUdpDeserializedMessage::GetAttachment() const
+TSharedPtr<IMessageAttachment, ESPMode::ThreadSafe> FUdpDeserializedMessage::GetAttachment() const
 {
 	return Attachment;
 }
diff --git a/Engine/Plugins/Messaging/UdpMessaging/Source/UdpMessaging/Private/Transport/UdpDeserializedMessage.h b/Engine/Plugins/Messaging/UdpMessaging/Source/UdpMessaging/Private/Transport/UdpDeserializedMessage.h
--- a/Engine/Plugins/Messaging/UdpMessaging/Source/UdpMessaging/Private/Transport/UdpDeserializedMessage.h
+++ b/Engine/Plugins/Messaging/UdpMessaging/Source/UdpMessaging/Private/Transport/UdpDeserializedMessage.h
@@ -3,6 +3,9 @@
 #pragma once
 
 
+class FReassembledUdpMessage;
+
+
 /**
  * Holds a deserialized message.
  */
@@ -16,7 +19,7 @@ public:
 	 *
 	 * @param InAttachment An optional message attachment.
 	 */
-	FUdpDeserializedMessage(const IMessageAttachmentPtr& InAttachment);
+	FUdpDeserializedMessage(const TSharedPtr<IMessageAttachment, ESPMode::ThreadSafe>& InAttachment);
 
 	/** Virtual destructor. */
 	virtual ~FUdpDeserializedMessage() override;
@@ -29,14 +32,14 @@ public:
 	 * @param ReassembledMessage The reassembled message to deserialize.
 	 * @return true on success, false otherwise.
 	 */
-	bool Deserialize(const FUdpReassembledMessageRef& ReassembledMessage);
+	bool Deserialize(const FReassembledUdpMessage& ReassembledMessage);
 
 public:
 
-	// IMessageContext interface
+	//~ IMessageContext interface
 
 	virtual const TMap<FName, FString>& GetAnnotations() const override;
-	virtual IMessageAttachmentPtr GetAttachment() const override;
+	virtual TSharedPtr<IMessageAttachment, ESPMode::ThreadSafe> GetAttachment() const override;
 	virtual const FDateTime& GetExpiration() const override;
 	virtual const void* GetMessage() const override;
 	virtual const TWeakObjectPtr<UScriptStruct>& GetMessageTypeInfo() const override;
@@ -54,7 +57,7 @@ private:
 	TMap<FName, FString> Annotations;
 
 	/** Holds a pointer to attached binary data. */
-	IMessageAttachmentPtr Attachment;
+	TSharedPtr<IMessageAttachment, ESPMode::ThreadSafe> Attachment;
 
 	/** Holds the expiration time. */
 	FDateTime Expiration;
@@ -77,10 +80,3 @@ private:
 	/** Holds the message's type information. */
 	TWeakObjectPtr<UScriptStruct> TypeInfo;
 };
-
-
-/** Type definition for shared pointers to instances of FUdpDeserializedMessage. */
-typedef TSharedPtr<FUdpDeserializedMessage, ESPMode::ThreadSafe> FUdpDeserializedMessagePtr;
-
-/** Type definition for shared references to instances of FUdpDeserializedMessage. */
-typedef TSharedRef<FUdpDeserializedMessage, ESPMode::ThreadSafe> FUdpDeserializedMessageRef;
diff --git a/Engine/Plugins/Messaging/UdpMessaging/Source/UdpMessaging/Private/Transport/UdpMessageBeacon.cpp b/Engine/Plugins/Messaging/UdpMessaging/Source/UdpMessaging/Private/Transport/UdpMessageBeacon.cpp
--- a/Engine/Plugins/Messaging/UdpMessaging/Source/UdpMessaging/Private/Transport/UdpMessageBeacon.cpp
+++ b/Engine/Plugins/Messaging/UdpMessaging/Source/UdpMessaging/Private/Transport/UdpMessageBeacon.cpp
@@ -1,6 +1,8 @@
 // Copyright 1998-2016 Epic Games, Inc. All Rights Reserved.
 
 #include "UdpMessagingPrivatePCH.h"
+#include "UdpMessageBeacon.h"
+#include "UdpMessageSegment.h"
 
 
 /* FUdpMessageHelloSender static initialization
@@ -14,7 +16,8 @@ const FTimespan FUdpMessageBeacon::MinimumInterval = FTimespan::FromMilliseconds
  *****************************************************************************/
 
 FUdpMessageBeacon::FUdpMessageBeacon(FSocket* InSocket, const FGuid& InSocketId, const FIPv4Endpoint& InMulticastEndpoint)
-	: LastEndpointCount(1)
+	: BeaconInterval(MinimumInterval)
+	, LastEndpointCount(1)
 	, LastHelloSent(FDateTime::MinValue())
 	, NextHelloTime(FDateTime::UtcNow())
 	, NodeId(InSocketId)
@@ -81,11 +84,11 @@ uint32 FUdpMessageBeacon::Run()
 
 		if (CurrentTime >= NextHelloTime)
 		{
-			SendSegment(EUdpMessageSegments::Hello);
-
 			// calculate the next send interval
 			BeaconInterval = FMath::Max(MinimumInterval, IntervalPerEndpoint * LastEndpointCount);
 			NextHelloTime = CurrentTime + BeaconInterval;
+
+			SendSegment(EUdpMessageSegments::Hello);
 		}
 
 		EndpointLeftEvent->Wait(NextHelloTime - CurrentTime);
@@ -106,7 +109,7 @@ void FUdpMessageBeacon::Stop()
 /* FUdpMessageHelloSender implementation
  *****************************************************************************/
 
-void FUdpMessageBeacon::SendSegment(EUdpMessageSegments::Type SegmentType)
+void FUdpMessageBeacon::SendSegment(EUdpMessageSegments SegmentType)
 {
 	FUdpMessageSegment::FHeader Header;
 	{
diff --git a/Engine/Plugins/Messaging/UdpMessaging/Source/UdpMessaging/Private/Transport/UdpMessageBeacon.h b/Engine/Plugins/Messaging/UdpMessaging/Source/UdpMessaging/Private/Transport/UdpMessageBeacon.h
--- a/Engine/Plugins/Messaging/UdpMessaging/Source/UdpMessaging/Private/Transport/UdpMessageBeacon.h
+++ b/Engine/Plugins/Messaging/UdpMessaging/Source/UdpMessaging/Private/Transport/UdpMessageBeacon.h
@@ -3,6 +3,9 @@
 #pragma once
 
 
+enum class EUdpMessageSegments : uint8;
+
+
 /**
  * Implements a beacon sender thread.
  *
@@ -46,7 +49,7 @@ public:
 
 public:
 
-	// FRunnable interface
+	//~ FRunnable interface
 
 	virtual bool Init() override;
 	virtual uint32 Run() override;
@@ -60,7 +63,7 @@ protected:
 	 *
 	 * @param SegmentType The type of segment to send (Hello or Bye).
 	 */
-	void SendSegment(EUdpMessageSegments::Type SegmentType);
+	void SendSegment(EUdpMessageSegments SegmentType);
 
 private:
 
diff --git a/Engine/Plugins/Messaging/UdpMessaging/Source/UdpMessaging/Private/Transport/UdpMessageProcessor.cpp b/Engine/Plugins/Messaging/UdpMessaging/Source/UdpMessaging/Private/Transport/UdpMessageProcessor.cpp
--- a/Engine/Plugins/Messaging/UdpMessaging/Source/UdpMessaging/Private/Transport/UdpMessageProcessor.cpp
+++ b/Engine/Plugins/Messaging/UdpMessaging/Source/UdpMessaging/Private/Transport/UdpMessageProcessor.cpp
@@ -1,6 +1,13 @@
 // Copyright 1998-2016 Epic Games, Inc. All Rights Reserved.
 
 #include "UdpMessagingPrivatePCH.h"
+#include "UdpMessageBeacon.h"
+#include "UdpMessageProcessor.h"
+#include "UdpMessageResequencer.h"
+#include "UdpMessageSegment.h"
+#include "UdpMessageSegmenter.h"
+#include "UdpMessagingSettings.h"
+#include "UdpReassembledMessage.h"
 
 
 /* FUdpMessageHelloSender static initialization
@@ -82,7 +89,7 @@ bool FUdpMessageProcessor::EnqueueInboundSegment(const FArrayReaderPtr& Data, co
 }
 
 
-bool FUdpMessageProcessor::EnqueueOutboundMessage(const FUdpSerializedMessageRef& SerializedMessage, const FGuid& Recipient)
+bool FUdpMessageProcessor::EnqueueOutboundMessage(const TSharedRef<FUdpSerializedMessage, ESPMode::ThreadSafe>& SerializedMessage, const FGuid& Recipient)
 {
 	if (!OutboundMessages.Enqueue(FOutboundMessage(SerializedMessage, Recipient)))
 	{
@@ -233,7 +240,7 @@ void FUdpMessageProcessor::ConsumeInboundSegments()
 				break;
 
 			default:
-				ProcessUnknownSegment(Segment, NodeInfo, Header.SegmentType);
+				ProcessUnknownSegment(Segment, NodeInfo, (uint8)Header.SegmentType);
 			}
 
 			NodeInfo.LastSegmentReceivedTime = CurrentTime;
@@ -337,7 +344,7 @@ void FUdpMessageProcessor::ProcessDataSegment(FInboundSegment& Segment, FNodeInf
 		return;
 	}
 
-	FUdpReassembledMessagePtr& ReassembledMessage = NodeInfo.ReassembledMessages.FindOrAdd(DataChunk.MessageId);
+	TSharedPtr<FReassembledUdpMessage, ESPMode::ThreadSafe>& ReassembledMessage = NodeInfo.ReassembledMessages.FindOrAdd(DataChunk.MessageId);
 
 	// Reassemble message
 	if (!ReassembledMessage.IsValid())
@@ -359,18 +366,18 @@ void FUdpMessageProcessor::ProcessDataSegment(FInboundSegment& Segment, FNodeInf
 	{
 		if (NodeInfo.NodeId.IsValid())
 		{
-			MessageReassembledDelegate.ExecuteIfBound(ReassembledMessage.ToSharedRef(), nullptr, NodeInfo.NodeId);
+			MessageReassembledDelegate.ExecuteIfBound(*ReassembledMessage, nullptr, NodeInfo.NodeId);
 		}
 	}
 	else if (NodeInfo.Resequencer.Resequence(ReassembledMessage))
 	{
-		FUdpReassembledMessagePtr ResequencedMessage;
+		TSharedPtr<FReassembledUdpMessage, ESPMode::ThreadSafe> ResequencedMessage;
 
 		while (NodeInfo.Resequencer.Pop(ResequencedMessage))
 		{
 			if (NodeInfo.NodeId.IsValid())
 			{
-				MessageReassembledDelegate.ExecuteIfBound(ResequencedMessage.ToSharedRef(), nullptr, NodeInfo.NodeId);
+				MessageReassembledDelegate.ExecuteIfBound(*ResequencedMessage, nullptr, NodeInfo.NodeId);
 			}
 		}
 	}
diff --git a/Engine/Plugins/Messaging/UdpMessaging/Source/UdpMessaging/Private/Transport/UdpMessageProcessor.h b/Engine/Plugins/Messaging/UdpMessaging/Source/UdpMessaging/Private/Transport/UdpMessageProcessor.h
--- a/Engine/Plugins/Messaging/UdpMessaging/Source/UdpMessaging/Private/Transport/UdpMessageProcessor.h
+++ b/Engine/Plugins/Messaging/UdpMessaging/Source/UdpMessaging/Private/Transport/UdpMessageProcessor.h
@@ -2,6 +2,16 @@
 
 #pragma once
 
+#include "Runnable.h"
+#include "UdpMessageResequencer.h"
+#include "UdpMessageSegment.h"
+
+
+class FReassembledUdpMessage;
+class FUdpMessageBeacon;
+class FUdpMessageSegmenter;
+class FUdpSerializedMessage;
+
 
 /**
  * Implements a message processor for UDP messages.
@@ -22,7 +32,7 @@ class FUdpMessageProcessor
 		FGuid NodeId;
 
 		/** Holds the collection of reassembled messages. */
-		TMap<int32, FUdpReassembledMessagePtr> ReassembledMessages;
+		TMap<int32, TSharedPtr<FReassembledUdpMessage, ESPMode::ThreadSafe>> ReassembledMessages;
 
 		/** Holds the message resequencer. */
 		FUdpMessageResequencer Resequencer;
@@ -70,11 +80,11 @@ class FUdpMessageProcessor
 	};
 
 
-	// Structure for outbound messages.
+	/** Structure for outbound messages. */
 	struct FOutboundMessage
 	{
 		/** Holds the serialized message. */
-		FUdpSerializedMessagePtr SerializedMessage;
+		TSharedPtr<FUdpSerializedMessage, ESPMode::ThreadSafe> SerializedMessage;
 
 		/** Holds the recipient. */
 		FGuid RecipientId;
@@ -83,7 +93,7 @@ class FUdpMessageProcessor
 		FOutboundMessage() { }
 
 		/** Creates and initializes a new instance. */
-		FOutboundMessage(const FUdpSerializedMessageRef& InSerializedMessage, const FGuid& InRecipientId)
+		FOutboundMessage(const TSharedRef<FUdpSerializedMessage, ESPMode::ThreadSafe>& InSerializedMessage, const FGuid& InRecipientId)
 			: SerializedMessage(InSerializedMessage)
 			, RecipientId(InRecipientId)
 		{ }
@@ -121,16 +131,19 @@ public:
 	 * @param Recipient The recipient's IPv4 endpoint.
 	 * @return true if the message was queued up, false otherwise.
 	 */
-	bool EnqueueOutboundMessage(const FUdpSerializedMessageRef& SerializedMessage, const FGuid& Recipient);
+	bool EnqueueOutboundMessage(const TSharedRef<FUdpSerializedMessage, ESPMode::ThreadSafe>& SerializedMessage, const FGuid& Recipient);
 
 public:
 
+	// @todo gmp: remove the need for this typedef
+	typedef TSharedPtr<IMessageAttachment, ESPMode::ThreadSafe> IMessageAttachmentPtr;
+
 	/**
 	 * Returns a delegate that is executed when message data has been received.
 	 *
 	 * @return The delegate.
 	 */
-	DECLARE_DELEGATE_ThreeParams(FOnMessageReassembled, const FUdpReassembledMessageRef& /*ReassembledMessage*/, const IMessageAttachmentPtr& /*Attachment*/, const FGuid& /*NodeId*/)
+	DECLARE_DELEGATE_ThreeParams(FOnMessageReassembled, const FReassembledUdpMessage& /*ReassembledMessage*/, const IMessageAttachmentPtr& /*Attachment*/, const FGuid& /*NodeId*/)
 	FOnMessageReassembled& OnMessageReassembled()
 	{
 		return MessageReassembledDelegate;
@@ -160,7 +173,7 @@ public:
 	
 public:
 
-	// FRunnable interface
+	//~ FRunnable interface
 
 	virtual bool Init() override;
 	virtual uint32 Run() override;
diff --git a/Engine/Plugins/Messaging/UdpMessaging/Source/UdpMessaging/Private/Transport/UdpMessageResequencer.h b/Engine/Plugins/Messaging/UdpMessaging/Source/UdpMessaging/Private/Transport/UdpMessageResequencer.h
--- a/Engine/Plugins/Messaging/UdpMessaging/Source/UdpMessaging/Private/Transport/UdpMessageResequencer.h
+++ b/Engine/Plugins/Messaging/UdpMessaging/Source/UdpMessaging/Private/Transport/UdpMessageResequencer.h
@@ -2,6 +2,8 @@
 
 #pragma once
 
+#include "UdpReassembledMessage.h"
+
 
 /**
  * Implements a re-sequencer for messages received over the UDP transport.
@@ -15,6 +17,8 @@ public:
 
 	/**
 	 * Creates and initializes a new message resequencer.
+	 *
+	 * @param InWindowSize The maximum resequencing window size.
 	 */
 	FUdpMessageResequencer(uint16 InWindowSize)
 		: NextSequence(1)
@@ -38,7 +42,7 @@ public:
 	 *
 	 * @return true if a message was returned, false otherwise.
 	 */
-	bool Pop(FUdpReassembledMessagePtr& OutMessage)
+	bool Pop(TSharedPtr<FReassembledUdpMessage, ESPMode::ThreadSafe>& OutMessage)
 	{
 		if (MessageHeap.HeapTop()->GetSequence() != NextSequence)
 		{
@@ -57,7 +61,7 @@ public:
 	 * @param Message The message to resequence.
 	 * @return true if the message is in sequence, false otherwise.
 	 */
-	bool Resequence(const FUdpReassembledMessagePtr& Message)
+	bool Resequence(const TSharedPtr<FReassembledUdpMessage, ESPMode::ThreadSafe>& Message)
 	{
 		MessageHeap.HeapPush(Message, FSequenceComparer());
 
@@ -76,7 +80,7 @@ private:
 	/** Helper for ordering messages by their sequence numbers. */
 	struct FSequenceComparer
 	{
-		bool operator()(const FUdpReassembledMessagePtr& A, const FUdpReassembledMessagePtr& B) const
+		bool operator()(const TSharedPtr<FReassembledUdpMessage, ESPMode::ThreadSafe>& A, const TSharedPtr<FReassembledUdpMessage, ESPMode::ThreadSafe>& B) const
 		{
 			return A->GetSequence() < B->GetSequence();
 		}
@@ -91,7 +95,7 @@ private:
 	uint64 HighestReceivedSequence;
 
 	/** Holds the messages that need to be resequenced. */
-	TArray<FUdpReassembledMessagePtr> MessageHeap;
+	TArray<TSharedPtr<FReassembledUdpMessage, ESPMode::ThreadSafe>> MessageHeap;
 
 	/** Holds the maximum resequence window size. */
 	uint16 WindowSize;
diff --git a/Engine/Plugins/Messaging/UdpMessaging/Source/UdpMessaging/Private/Transport/UdpMessageSegmenter.cpp b/Engine/Plugins/Messaging/UdpMessaging/Source/UdpMessaging/Private/Transport/UdpMessageSegmenter.cpp
--- a/Engine/Plugins/Messaging/UdpMessaging/Source/UdpMessaging/Private/Transport/UdpMessageSegmenter.cpp
+++ b/Engine/Plugins/Messaging/UdpMessaging/Source/UdpMessaging/Private/Transport/UdpMessageSegmenter.cpp
@@ -1,6 +1,7 @@
 // Copyright 1998-2016 Epic Games, Inc. All Rights Reserved.
 
 #include "UdpMessagingPrivatePCH.h"
+#include "UdpMessageSegmenter.h"
 
 
 /* FUdpMessageSegmenter structors
diff --git a/Engine/Plugins/Messaging/UdpMessaging/Source/UdpMessaging/Private/Transport/UdpMessageSegmenter.h b/Engine/Plugins/Messaging/UdpMessaging/Source/UdpMessaging/Private/Transport/UdpMessageSegmenter.h
--- a/Engine/Plugins/Messaging/UdpMessaging/Source/UdpMessaging/Private/Transport/UdpMessageSegmenter.h
+++ b/Engine/Plugins/Messaging/UdpMessaging/Source/UdpMessaging/Private/Transport/UdpMessageSegmenter.h
@@ -2,6 +2,8 @@
 
 #pragma once
 
+#include "UdpSerializedMessage.h"
+
 
 /**
  * Implements a message segmenter.
@@ -23,7 +25,7 @@ public:
 	 *
 	 * @param InMessage The serialized message to segment.
 	 */
-	FUdpMessageSegmenter(const FUdpSerializedMessageRef& InSerializedMessage, uint16 InSegmentSize)
+	FUdpMessageSegmenter(const TSharedRef<FUdpSerializedMessage, ESPMode::ThreadSafe>& InSerializedMessage, uint16 InSegmentSize)
 		: MessageReader(nullptr)
 		, SegmentSize(InSegmentSize)
 		, SerializedMessage(InSerializedMessage)
@@ -140,5 +142,5 @@ private:
 	uint16 SegmentSize;
 
 	/** Holds the message. */
-	FUdpSerializedMessagePtr SerializedMessage;
+	TSharedPtr<FUdpSerializedMessage, ESPMode::ThreadSafe> SerializedMessage;
 };
diff --git a/Engine/Plugins/Messaging/UdpMessaging/Source/UdpMessaging/Private/Transport/UdpMessageTransport.cpp b/Engine/Plugins/Messaging/UdpMessaging/Source/UdpMessaging/Private/Transport/UdpMessageTransport.cpp
--- a/Engine/Plugins/Messaging/UdpMessaging/Source/UdpMessaging/Private/Transport/UdpMessageTransport.cpp
+++ b/Engine/Plugins/Messaging/UdpMessaging/Source/UdpMessaging/Private/Transport/UdpMessageTransport.cpp
@@ -1,7 +1,11 @@
 // Copyright 1998-2016 Epic Games, Inc. All Rights Reserved.
 
 #include "UdpMessagingPrivatePCH.h"
+#include "UdpDeserializedMessage.h"
+#include "UdpMessageProcessor.h"
+#include "UdpMessageTransport.h"
 #include "UdpSerializeMessageTask.h"
+#include "UdpSerializedMessage.h"
 
 
 /* FUdpMessageTransport structors
@@ -141,7 +145,7 @@ void FUdpMessageTransport::StopTransport()
 }
 
 
-bool FUdpMessageTransport::TransportMessage(const IMessageContextRef& Context, const TArray<FGuid>& Recipients)
+bool FUdpMessageTransport::TransportMessage(const TSharedRef<IMessageContext, ESPMode::ThreadSafe>& Context, const TArray<FGuid>& Recipients)
 {
 	if (MessageProcessor == nullptr)
 	{
@@ -153,7 +157,7 @@ bool FUdpMessageTransport::TransportMessage(const IMessageContextRef& Context, c
 		return false;
 	}
 
-	FUdpSerializedMessageRef SerializedMessage = MakeShareable(new FUdpSerializedMessage());
+	TSharedRef<FUdpSerializedMessage, ESPMode::ThreadSafe> SerializedMessage = MakeShareable(new FUdpSerializedMessage());
 
 	if (Recipients.Num() == 0)
 	{
@@ -178,10 +182,10 @@ bool FUdpMessageTransport::TransportMessage(const IMessageContextRef& Context, c
 /* FUdpMessageTransport event handlers
  *****************************************************************************/
 
-void FUdpMessageTransport::HandleProcessorMessageReassembled(const FUdpReassembledMessageRef& ReassembledMessage, const IMessageAttachmentPtr& Attachment, const FGuid& NodeId)
+void FUdpMessageTransport::HandleProcessorMessageReassembled(const FReassembledUdpMessage& ReassembledMessage, const TSharedPtr<IMessageAttachment, ESPMode::ThreadSafe>& Attachment, const FGuid& NodeId)
 {
 	// @todo gmp: move message deserialization into an async task
-	FUdpDeserializedMessageRef DeserializedMessage = MakeShareable(new FUdpDeserializedMessage(Attachment));
+	TSharedRef<FUdpDeserializedMessage, ESPMode::ThreadSafe> DeserializedMessage = MakeShareable(new FUdpDeserializedMessage(Attachment));
 
 	if (DeserializedMessage->Deserialize(ReassembledMessage))
 	{
diff --git a/Engine/Plugins/Messaging/UdpMessaging/Source/UdpMessaging/Private/Transport/UdpMessageTransport.h b/Engine/Plugins/Messaging/UdpMessaging/Source/UdpMessaging/Private/Transport/UdpMessageTransport.h
--- a/Engine/Plugins/Messaging/UdpMessaging/Source/UdpMessaging/Private/Transport/UdpMessageTransport.h
+++ b/Engine/Plugins/Messaging/UdpMessaging/Source/UdpMessaging/Private/Transport/UdpMessageTransport.h
@@ -2,6 +2,16 @@
 
 #pragma once
 
+#include "IMessageTransport.h"
+
+
+struct FIPv4Endpoint;
+class FReassembledUdpMessage;
+class FSocket;
+class FUdpMessageProcessor;
+class FUdpSocketReceiver;
+class ISocketSubsystem;
+
 
 /**
  * Implements a message transport technology using an UDP network connection.
@@ -29,7 +39,7 @@ public:
 
 public:
 
-	// IMessageTransport interface
+	//~ IMessageTransport interface
 
 	virtual FName GetDebugName() const override
 	{
@@ -53,12 +63,12 @@ public:
 
 	virtual bool StartTransport() override;
 	virtual void StopTransport() override;
-	virtual bool TransportMessage(const IMessageContextRef& Context, const TArray<FGuid>& Recipients) override;
+	virtual bool TransportMessage(const TSharedRef<IMessageContext, ESPMode::ThreadSafe>& Context, const TArray<FGuid>& Recipients) override;
 
 private:
 
 	/** Handles received transport messages. */
-	void HandleProcessorMessageReassembled(const FUdpReassembledMessageRef& ReassembledMessage, const IMessageAttachmentPtr& Attachment, const FGuid& NodeId);
+	void HandleProcessorMessageReassembled(const FReassembledUdpMessage& ReassembledMessage, const TSharedPtr<IMessageAttachment, ESPMode::ThreadSafe>& Attachment, const FGuid& NodeId);
 
 	/** Handles discovered transport endpoints. */
 	void HandleProcessorNodeDiscovered(const FGuid& DiscoveredNodeId)
diff --git a/Engine/Plugins/Messaging/UdpMessaging/Source/UdpMessaging/Private/Transport/UdpReassembledMessage.h b/Engine/Plugins/Messaging/UdpMessaging/Source/UdpMessaging/Private/Transport/UdpReassembledMessage.h
--- a/Engine/Plugins/Messaging/UdpMessaging/Source/UdpMessaging/Private/Transport/UdpReassembledMessage.h
+++ b/Engine/Plugins/Messaging/UdpMessaging/Source/UdpMessaging/Private/Transport/UdpReassembledMessage.h
@@ -3,13 +3,6 @@
 #pragma once
 
 
-/** Type definition for shared pointers to instances of FUdpMessage. */
-typedef TSharedPtr<class FReassembledUdpMessage, ESPMode::ThreadSafe> FUdpReassembledMessagePtr;
-
-/** Type definition for shared references to instances of FUdpMessage. */
-typedef TSharedRef<class FReassembledUdpMessage, ESPMode::ThreadSafe> FUdpReassembledMessageRef;
-
-
 /**
  * Implements a reassembled message.
  */
diff --git a/Engine/Plugins/Messaging/UdpMessaging/Source/UdpMessaging/Private/Transport/UdpSerializeMessageTask.cpp b/Engine/Plugins/Messaging/UdpMessaging/Source/UdpMessaging/Private/Transport/UdpSerializeMessageTask.cpp
--- a/Engine/Plugins/Messaging/UdpMessaging/Source/UdpMessaging/Private/Transport/UdpSerializeMessageTask.cpp
+++ b/Engine/Plugins/Messaging/UdpMessaging/Source/UdpMessaging/Private/Transport/UdpSerializeMessageTask.cpp
@@ -1,7 +1,8 @@
 // Copyright 1998-2016 Epic Games, Inc. All Rights Reserved.
 
 #include "UdpMessagingPrivatePCH.h"
-#include "TaskGraphInterfaces.h"
+#include "UdpSerializeMessageTask.h"
+#include "UdpSerializedMessage.h"
 
 
 /* FUdpSerializeMessageTask interface
diff --git a/Engine/Plugins/Messaging/UdpMessaging/Source/UdpMessaging/Private/Transport/UdpSerializeMessageTask.h b/Engine/Plugins/Messaging/UdpMessaging/Source/UdpMessaging/Private/Transport/UdpSerializeMessageTask.h
--- a/Engine/Plugins/Messaging/UdpMessaging/Source/UdpMessaging/Private/Transport/UdpSerializeMessageTask.h
+++ b/Engine/Plugins/Messaging/UdpMessaging/Source/UdpMessaging/Private/Transport/UdpSerializeMessageTask.h
@@ -2,8 +2,11 @@
 
 #pragma once
 
-#include "JsonStructSerializerBackend.h"
-#include "StructSerializer.h"
+#include "TaskGraphInterfaces.h"
+
+
+class FUdpSerializedMessage;
+class IMessageContext;
 
 
 /**
@@ -19,7 +22,7 @@ public:
 	 * @param InMessageContext The context of the message to serialize.
 	 * @param InSerializedMessage Will hold the serialized message data.
 	 */
-	FUdpSerializeMessageTask(IMessageContextRef InMessageContext, FUdpSerializedMessageRef InSerializedMessage)
+	FUdpSerializeMessageTask(TSharedRef<IMessageContext, ESPMode::ThreadSafe> InMessageContext, TSharedRef<FUdpSerializedMessage, ESPMode::ThreadSafe> InSerializedMessage)
 		: MessageContext(InMessageContext)
 		, SerializedMessage(InSerializedMessage)
 	{ }
@@ -58,8 +61,8 @@ public:
 private:
 
 	/** Holds the context of the message to serialize. */
-	IMessageContextRef MessageContext;
+	TSharedRef<IMessageContext, ESPMode::ThreadSafe> MessageContext;
 
 	/** Holds a reference to the serialized message data. */
-	FUdpSerializedMessageRef SerializedMessage;
+	TSharedRef<FUdpSerializedMessage, ESPMode::ThreadSafe> SerializedMessage;
 };
diff --git a/Engine/Plugins/Messaging/UdpMessaging/Source/UdpMessaging/Private/Transport/UdpSerializedMessage.h b/Engine/Plugins/Messaging/UdpMessaging/Source/UdpMessaging/Private/Transport/UdpSerializedMessage.h
--- a/Engine/Plugins/Messaging/UdpMessaging/Source/UdpMessaging/Private/Transport/UdpSerializedMessage.h
+++ b/Engine/Plugins/Messaging/UdpMessaging/Source/UdpMessaging/Private/Transport/UdpSerializedMessage.h
@@ -2,6 +2,9 @@
 
 #pragma once
 
+#include "MemoryReader.h"
+#include "MemoryWriter.h"
+
 
 /**
  * Enumerates possibly states of a serialized message.
@@ -107,10 +110,3 @@ private:
 	/** Holds a delegate that is invoked when the data's state changed. */
 	FSimpleDelegate StateChangedDelegate;
 };
-
-
-/** Type definition for shared pointers to instances of FUdpSerializedMessage. */
-typedef TSharedPtr<FUdpSerializedMessage, ESPMode::ThreadSafe> FUdpSerializedMessagePtr;
-
-/** Type definition for shared references to instances of FUdpSerializedMessage. */
-typedef TSharedRef<FUdpSerializedMessage, ESPMode::ThreadSafe> FUdpSerializedMessageRef;
diff --git a/Engine/Plugins/Messaging/UdpMessaging/Source/UdpMessaging/Private/Tunnel/UdpMessageTunnel.cpp b/Engine/Plugins/Messaging/UdpMessaging/Source/UdpMessaging/Private/Tunnel/UdpMessageTunnel.cpp
--- a/Engine/Plugins/Messaging/UdpMessaging/Source/UdpMessaging/Private/Tunnel/UdpMessageTunnel.cpp
+++ b/Engine/Plugins/Messaging/UdpMessaging/Source/UdpMessaging/Private/Tunnel/UdpMessageTunnel.cpp
@@ -4,6 +4,11 @@
 
 #if PLATFORM_DESKTOP
 
+#include "UdpMessageSegment.h"
+#include "UdpMessageTunnelConnection.h"
+#include "UdpMessageTunnel.h"
+
+
 /* FUdpMessageTunnel structors
  *****************************************************************************/
 
@@ -134,7 +139,7 @@ bool FUdpMessageTunnel::Connect(const FIPv4Endpoint& RemoteEndpoint)
 }
 
 
-int32 FUdpMessageTunnel::GetConnections(TArray<IUdpMessageTunnelConnectionPtr>& OutConnections)
+int32 FUdpMessageTunnel::GetConnections(TArray<TSharedPtr<IUdpMessageTunnelConnection>>& OutConnections)
 {
 	FScopeLock Lock(&CriticalSection);
 
diff --git a/Engine/Plugins/Messaging/UdpMessaging/Source/UdpMessaging/Private/Tunnel/UdpMessageTunnel.h b/Engine/Plugins/Messaging/UdpMessaging/Source/UdpMessaging/Private/Tunnel/UdpMessageTunnel.h
--- a/Engine/Plugins/Messaging/UdpMessaging/Source/UdpMessaging/Private/Tunnel/UdpMessageTunnel.h
+++ b/Engine/Plugins/Messaging/UdpMessaging/Source/UdpMessaging/Private/Tunnel/UdpMessageTunnel.h
@@ -5,7 +5,6 @@
 #include "IUdpMessageTunnel.h"
 
 
-// forward declarations
 class FRunnableThread;
 class FSocket;
 
@@ -45,7 +44,7 @@ public:
 
 public:
 
-	// FRunnable interface
+	//~ FRunnable interface
 
 	virtual bool Init() override;
 	virtual uint32 Run() override;
@@ -54,10 +53,10 @@ public:
 
 public:
 
-	// IUdpMessageTunnel interface
+	//~ IUdpMessageTunnel interface
 
 	virtual bool Connect(const FIPv4Endpoint& RemoteEndpoint) override;
-	virtual int32 GetConnections(TArray<IUdpMessageTunnelConnectionPtr>& OutConnections) override;
+	virtual int32 GetConnections(TArray<TSharedPtr<IUdpMessageTunnelConnection>>& OutConnections) override;
 	virtual uint64 GetTotalInboundBytes() const override;
 	virtual uint64 GetTotalOutboundBytes() const override;
 	virtual bool IsServerRunning() const override;
diff --git a/Engine/Plugins/Messaging/UdpMessaging/Source/UdpMessaging/Private/Tunnel/UdpMessageTunnelConnection.cpp b/Engine/Plugins/Messaging/UdpMessaging/Source/UdpMessaging/Private/Tunnel/UdpMessageTunnelConnection.cpp
--- a/Engine/Plugins/Messaging/UdpMessaging/Source/UdpMessaging/Private/Tunnel/UdpMessageTunnelConnection.cpp
+++ b/Engine/Plugins/Messaging/UdpMessaging/Source/UdpMessaging/Private/Tunnel/UdpMessageTunnelConnection.cpp
@@ -1,8 +1,10 @@
 // Copyright 1998-2016 Epic Games, Inc. All Rights Reserved.
 
+#include "UdpMessagingPrivatePCH.h"
+
 #if PLATFORM_DESKTOP
 
-#include "UdpMessagingPrivatePCH.h"
+#include "UdpMessageTunnelConnection.h"
 
 
 /* FUdpMessageTunnelConnection structors
