Messaging: Fix non unity build

diff --git a/Engine/Plugins/Messaging/MessagingDebugger/Source/MessagingDebugger/Private/MessagingDebuggerPCH.h b/Engine/Plugins/Messaging/MessagingDebugger/Source/MessagingDebugger/Private/MessagingDebuggerPCH.h
--- a/Engine/Plugins/Messaging/MessagingDebugger/Source/MessagingDebugger/Private/MessagingDebuggerPCH.h
+++ b/Engine/Plugins/Messaging/MessagingDebugger/Source/MessagingDebugger/Private/MessagingDebuggerPCH.h
@@ -6,6 +6,7 @@
 #include "Runtime/Core/Public/Features/IModularFeatures.h"
 #include "Runtime/Core/Public/Modules/ModuleManager.h"
 #include "Runtime/Json/Public/Json.h"
+#include "Runtime/Messaging/Public/IMessageBus.h"
 #include "Runtime/Messaging/Public/IMessageTracer.h"
 #include "Runtime/Messaging/Public/IMessageTracerBreakpoint.h"
 #include "Runtime/Serialization/Public/StructSerializer.h"
