Messaging: Fix non unity build

diff --git a/Engine/Plugins/Messaging/MessagingDebugger/Source/MessagingDebugger/Private/MessagingDebuggerPCH.h b/Engine/Plugins/Messaging/MessagingDebugger/Source/MessagingDebugger/Private/MessagingDebuggerPCH.h
--- a/Engine/Plugins/Messaging/MessagingDebugger/Source/MessagingDebugger/Private/MessagingDebuggerPCH.h
+++ b/Engine/Plugins/Messaging/MessagingDebugger/Source/MessagingDebugger/Private/MessagingDebuggerPCH.h
@@ -9,6 +9,7 @@
 #include "Runtime/Messaging/Public/IMessageBus.h"
 #include "Runtime/Messaging/Public/IMessageTracer.h"
 #include "Runtime/Messaging/Public/IMessageTracerBreakpoint.h"
+#include "Runtime/Messaging/Public/IMessagingModule.h"
 #include "Runtime/Serialization/Public/StructSerializer.h"
 #include "Runtime/Serialization/Public/Backends/JsonStructSerializerBackend.h"
 #include "Runtime/Slate/Public/SlateBasics.h"
